Integrating //UE5/Dev-Cooker @ 16678003 to //UE5/Main (Zousar.Shaker-YEG-0943-Quaternary)

diff --git a/Engine/Source/Runtime/CookOnTheFly/Private/CookOnTheFlyConnectionServer.cpp b/Engine/Source/Runtime/CookOnTheFly/Private/CookOnTheFlyConnectionServer.cpp
new file mode 100644
--- /dev/null
+++ b/Engine/Source/Runtime/CookOnTheFly/Private/CookOnTheFlyConnectionServer.cpp
@@ -0,0 +1,404 @@
+// Copyright Epic Games, Inc. All Rights Reserved.
+
+#include "CookOnTheFly.h"
+#include "HAL/PlatformMisc.h"
+#include "HAL/PlatformTime.h"
+#include "Async/Async.h"
+#include "IPAddress.h"
+#include "SocketSubsystem.h"
+#include "Sockets.h"
+#include "NetworkMessage.h"
+#include "Serialization/BufferReader.h"
+#include "Misc/DateTime.h"
+#include "Misc/OutputDeviceRedirector.h"
+
+DEFINE_LOG_CATEGORY_STATIC(LogCotfConnectionServer, Log, All);
+
+class FCookOnTheFlyConnectionServer final
+	: public UE::Cook::ICookOnTheFlyConnectionServer
+{
+	static constexpr uint32 ServerSenderId = ~uint32(0);
+	static constexpr double HeartbeatTimeoutInSeconds = 120.0;
+
+public:
+	FCookOnTheFlyConnectionServer(UE::Cook::FCookOnTheFlyServerOptions InOptions)
+		: Options(MoveTemp(InOptions)) { }
+
+	~FCookOnTheFlyConnectionServer()
+	{
+		StopServer();
+	}
+
+	virtual bool StartServer() override
+	{
+		const int32 Port = Options.Port > 0 ? Options.Port : UE::Cook::DefaultCookOnTheFlyServingPort; 
+		UE_LOG(LogCotfConnectionServer, Log, TEXT("Starting COTF server on port '%d'"), Port);
+
+		check(!bIsRunning);
+		bIsRunning = false;
+		bStopRequested = false;
+
+		ISocketSubsystem& SocketSubsystem = *ISocketSubsystem::Get();
+
+		ListenAddr = SocketSubsystem.GetLocalBindAddr(*GLog);
+		ListenAddr->SetPort(Port);
+
+		// create a server TCP socket
+		Socket = SocketSubsystem.CreateSocket(NAME_Stream, TEXT("COTF-Server"), ListenAddr->GetProtocolType());
+
+		if(!Socket)
+		{
+			UE_LOG(LogCotfConnectionServer, Error, TEXT("Could not create listen socket"));
+			return false;
+		}
+
+		Socket->SetReuseAddr();
+
+		if (!Socket->Bind(*ListenAddr))
+		{
+			UE_LOG(LogCotfConnectionServer, Error, TEXT("Failed to bind socket to addres '%s'"), *ListenAddr->ToString(true));
+			return false;
+		}
+
+		if (!Socket->Listen(16))
+		{
+			UE_LOG(LogCotfConnectionServer, Warning, TEXT("Failed to listen on address '%s'"), *ListenAddr->ToString(true));
+			return false;
+		}
+
+		ListenAddr->SetPort(Socket->GetPortNo());
+
+		ServerThread = AsyncThread([this] { return ServerThreadEntry(); },  8 * 1024, TPri_AboveNormal);
+
+		UE_LOG(LogCotfConnectionServer, Display, TEXT("COTF server is ready for client(s) on '%s'!"), *ListenAddr->ToString(true));
+
+		return true;
+	}
+
+private:
+
+	struct FClient
+	{
+		FSocket* Socket = nullptr;
+		TSharedPtr<FInternetAddr> Addr;
+		TSharedPtr<FInternetAddr> PeerAddr;
+		TFuture<void> Thread;
+		TAtomic<bool> bIsRunning { false };
+		TAtomic<bool> bStopRequested { false };
+		TAtomic<double> LastActivityTime {0};
+		uint32 ClientId = 0;
+		FName PlatformName;
+	};
+
+	virtual void StopServer() override
+	{
+		if (bIsRunning && !bStopRequested)
+		{
+			bStopRequested = true;
+			
+			ISocketSubsystem& SocketSubsystem = *ISocketSubsystem::Get();
+
+			while (bIsRunning)
+			{
+				FPlatformProcess::Sleep(0.25f);
+			}
+
+			for (TUniquePtr<FClient>& Client : Clients)
+			{
+				Client->bStopRequested = true;
+				Client->Socket->Close();
+				Client->Thread.Wait();
+				SocketSubsystem.DestroySocket(Client->Socket);
+			}
+
+			Clients.Empty();
+		}
+	}
+
+	virtual bool BroadcastMessage(const UE::Cook::FCookOnTheFlyMessage& Message, const FName& PlatformName = NAME_None) override
+	{
+		using namespace UE::Cook;
+
+		FCookOnTheFlyMessageHeader Header = Message.GetHeader();
+
+		Header.MessageType		= Header.MessageType | ECookOnTheFlyMessage::Message;
+		Header.MessageStatus	= ECookOnTheFlyMessageStatus::Ok;
+		Header.SenderId			= ServerSenderId;
+		Header.CorrelationId	= NextCorrelationId++;
+		Header.Timestamp		= FDateTime::UtcNow().GetTicks();
+
+		FBufferArchive MessagePayload;
+		MessagePayload.Reserve(Message.TotalSize());
+
+		MessagePayload << Header;
+		MessagePayload << const_cast<TArray<uint8>&>(Message.GetBody());
+
+		UE_LOG(LogCotfConnectionServer, Verbose, TEXT("Sending: %s, Size='%lld'"), *Header.ToString(), Message.TotalSize());
+
+		TArray<FClient*, TInlineAllocator<4>> ClientsToBroadcast;
+		{
+			FScopeLock _(&ClintsCriticalSection);
+
+			for (TUniquePtr<FClient>& Client : Clients)
+			{
+				if (PlatformName.IsNone() || Client->PlatformName == PlatformName)
+				{
+					ClientsToBroadcast.Add(Client.Get());
+				}
+			}
+		}
+
+		bool bBroadcasted = true;
+		for (FClient* Client : ClientsToBroadcast)
+		{
+			if (!FNFSMessageHeader::WrapAndSendPayload(MessagePayload, FSimpleAbstractSocket_FSocket(Client->Socket)))
+			{
+				UE_LOG(LogCotfConnectionServer, Warning, TEXT("Failed to send message '%s' to client '%s' (Id='%d', Platform='%s')"),
+					LexToString(Message.GetHeader().MessageType), *Client->PeerAddr->ToString(true), Client->ClientId, *Client->PlatformName.ToString());
+
+				bBroadcasted = false;
+			}
+		}
+
+		return bBroadcasted;
+	}
+
+	void ServerThreadEntry()
+	{
+		using namespace UE::Cook;
+
+		ISocketSubsystem& SocketSubsystem = *ISocketSubsystem::Get();
+		bIsRunning = true;
+
+		while (!bStopRequested)
+		{
+			bool bIsReady = false;
+			if (Socket->WaitForPendingConnection(bIsReady, FTimespan::FromSeconds(0.25f)))
+			{
+				if (bIsReady)
+				{
+					if (FSocket* ClientSocket = Socket->Accept(TEXT("COTF-Client")))
+					{
+						TUniquePtr<FClient> Client = MakeUnique<FClient>();
+						Client->Socket = ClientSocket;
+						Client->Addr = SocketSubsystem.CreateInternetAddr();
+						Client->PeerAddr = SocketSubsystem.CreateInternetAddr();
+						ClientSocket->GetAddress(*Client->Addr);
+						ClientSocket->GetPeerAddress(*Client->PeerAddr);
+						Client->ClientId = NextClientId++;
+						Client->bIsRunning = true;
+						Client->LastActivityTime = FPlatformTime::Seconds();
+						Client->Thread = AsyncThread([this, ConnectedClient = Client.Get()]()
+						{ 
+							ClientThreadEntry(ConnectedClient);
+						});
+
+						UE_LOG(LogCotfConnectionServer, Display, TEXT("New client connected from address '%s' (ID='%d')"), *Client->PeerAddr->ToString(true), Client->ClientId);
+						{
+							FScopeLock _(&ClintsCriticalSection);
+							Clients.Add(MoveTemp(Client));
+						}
+					}
+				}
+			}
+			else
+			{
+				FPlatformProcess::Sleep(0.25f);
+			}
+
+			{
+				FScopeLock _(&ClintsCriticalSection);
+				for (auto It = Clients.CreateIterator(); It; ++It)
+				{
+					TUniquePtr<FClient>& Client = *It;
+
+					const double SecondsSinceLastActivity = FPlatformTime::Seconds() - Client->LastActivityTime.Load();
+
+					if (SecondsSinceLastActivity > HeartbeatTimeoutInSeconds)
+					{
+						Client->LastActivityTime = FPlatformTime::Seconds();
+
+						UE_LOG(LogCotfConnectionServer, Display, TEXT("Sending hearbeat message, ClientId='%d', Platform='%s', Address='%s', IdleTime='%.2llf's"),
+							Client->ClientId, *Client->PlatformName.ToString(), *Client->PeerAddr->ToString(true), HeartbeatTimeoutInSeconds);
+						
+						FCookOnTheFlyMessage HeartbeatRequest(ECookOnTheFlyMessage::Heartbeat | ECookOnTheFlyMessage::Request);
+						FCookOnTheFlyMessageHeader& Header = HeartbeatRequest.GetHeader();
+
+						Header.MessageStatus	= ECookOnTheFlyMessageStatus::Ok;
+						Header.SenderId			= ServerSenderId;
+						Header.CorrelationId	= Client->ClientId;
+						Header.Timestamp		= FDateTime::UtcNow().GetTicks();
+						
+						FBufferArchive RequestPayload;
+						RequestPayload.Reserve(HeartbeatRequest.TotalSize());
+						RequestPayload << HeartbeatRequest;
+
+						if (!FNFSMessageHeader::WrapAndSendPayload(RequestPayload, FSimpleAbstractSocket_FSocket(Client->Socket)))
+						{
+							Client->bIsRunning = false;
+							UE_LOG(LogCotfConnectionServer, Display, TEXT("Heartbeat [Failed]"));
+						}
+					}
+
+					if (!Client->bIsRunning)
+					{
+						UE_LOG(LogCotfConnectionServer, Display, TEXT("Closing connection to client on address '%s' (Id='%d', Platform='%s')"),
+							*Client->PeerAddr->ToString(true), Client->ClientId, *Client->PlatformName.ToString());
+						
+						Options.HandleClientConnection(UE::Cook::FCookOnTheFlyClient { Client->ClientId, Client->PlatformName }, UE::Cook::ECookOnTheFlyConnectionStatus::Disconnected);
+
+						Client->Socket->Close();
+						Client->Thread.Wait();
+						SocketSubsystem.DestroySocket(Client->Socket);
+						It.RemoveCurrent();
+					}
+				}
+			}
+		}
+
+		bIsRunning = false;
+	}
+
+	void ClientThreadEntry(FClient* Client)
+	{
+		while (!bStopRequested && !Client->bStopRequested)
+		{
+			Client->LastActivityTime = FPlatformTime::Seconds();
+			if (!ProcesseRequest(*Client))
+			{
+				break;
+			}
+		}
+
+		Client->bIsRunning = false;
+	}
+
+	bool ProcesseRequest(FClient& Client)
+	{
+		using namespace UE::Cook;
+		
+		FArrayReader RequestPayload;
+		if (!FNFSMessageHeader::ReceivePayload(RequestPayload, FSimpleAbstractSocket_FSocket(Client.Socket)))
+		{
+			UE_LOG(LogCotfConnectionServer, Warning, TEXT("Unable to receive request from client"));
+			return false;
+		}
+
+		FCookOnTheFlyRequest Request;
+		RequestPayload << Request;
+
+		UE_LOG(LogCotfConnectionServer, Verbose, TEXT("Received: %s, Size='%lld'"), *Request.GetHeader().ToString(), Request.TotalSize());
+
+		EnumRemoveFlags(Request.GetHeader().MessageType, ECookOnTheFlyMessage::TypeFlags);
+
+		FCookOnTheFlyResponse Response;
+		bool bRequestOk = false;
+		bool bIsResponse = false;
+
+		switch (Request.GetHeader().MessageType)
+		{
+			case ECookOnTheFlyMessage::Handshake:
+			{
+				ProcessHandshake(Client, Request, Response);
+				bRequestOk = Options.HandleClientConnection(UE::Cook::FCookOnTheFlyClient { Client.ClientId, Client.PlatformName }, UE::Cook::ECookOnTheFlyConnectionStatus::Connected);
+				break;
+			}
+			case ECookOnTheFlyMessage::Heartbeat:
+			{
+				const bool bHeartbeatOk = Request.GetHeader().CorrelationId == Client.ClientId;
+
+				UE_LOG(LogCotfConnectionServer, Display, TEXT("Heartbeat [%s], ClientId='%d', Platform='%s', Address='%s'"),
+					bHeartbeatOk ? TEXT("Ok") : TEXT("Failed"), Client.ClientId, *Client.PlatformName.ToString(), *Client.PeerAddr->ToString(true));
+
+				bRequestOk = bHeartbeatOk;
+				bIsResponse  = true;
+				break;
+			}
+			default:
+			{
+				bRequestOk = Options.HandleRequest(UE::Cook::FCookOnTheFlyClient { Client.ClientId, Client.PlatformName }, Request, Response);
+				break;
+			}
+		}
+
+		if (bRequestOk && !bIsResponse)
+		{
+			FCookOnTheFlyMessageHeader& ResponseHeader = Response.GetHeader();
+
+			ResponseHeader.MessageType		= Request.GetHeader().MessageType | ECookOnTheFlyMessage::Response;
+			ResponseHeader.SenderId			= ServerSenderId;
+			ResponseHeader.CorrelationId	= Request.GetHeader().CorrelationId;
+			ResponseHeader.Timestamp		= FDateTime::UtcNow().GetTicks();
+
+			Response.SetHeader(ResponseHeader);
+
+			FBufferArchive ResponsePayload;
+			ResponsePayload.Reserve(Response.TotalSize());
+
+			ResponsePayload << Response;
+			bRequestOk = FNFSMessageHeader::WrapAndSendPayload(ResponsePayload, FSimpleAbstractSocket_FSocket(Client.Socket));
+		}
+
+		return bRequestOk;
+	}
+
+	void ProcessHandshake(FClient& Client, UE::Cook::FCookOnTheFlyRequest& HandshakeRequest, UE::Cook::FCookOnTheFlyResponse& Response)
+	{
+		FString PlatformName;
+		FString ProjectName;
+		
+		{
+			TUniquePtr<FArchive> Ar = HandshakeRequest.ReadBody();
+			*Ar << PlatformName;
+			*Ar << ProjectName;
+		}
+
+		if (PlatformName.Len())
+		{
+			Client.PlatformName = FName(*PlatformName);
+
+			Response.SetBodyTo(Client.ClientId);
+			Response.SetStatus(UE::Cook::ECookOnTheFlyMessageStatus::Ok);
+		}
+		else
+		{
+			Response.SetStatus(UE::Cook::ECookOnTheFlyMessageStatus::Error);
+		}
+	}
+
+	FClient* GetClientById(uint32 ClientId)
+	{
+		FScopeLock _(&ClintsCriticalSection);
+
+		for (TUniquePtr<FClient>& Client : Clients)
+		{
+			if (Client->ClientId == ClientId)
+			{
+				return Client.Get();
+			}
+		}
+
+		return nullptr;
+	}
+
+	UE::Cook::FCookOnTheFlyServerOptions Options;
+	TFuture<void> ServerThread;
+	TSharedPtr<FInternetAddr> ListenAddr;
+	FSocket* Socket = nullptr;
+	FCriticalSection ClintsCriticalSection;
+	TArray<TUniquePtr<FClient>> Clients;
+	TAtomic<bool> bIsRunning { false };
+	TAtomic<bool> bStopRequested { false };
+	uint32 NextClientId = 1;
+	TAtomic<uint32> NextCorrelationId { 1 };
+};
+
+namespace UE { namespace Cook
+{
+
+TUniquePtr<ICookOnTheFlyConnectionServer> MakeCookOnTheFlyConnectionServer(FCookOnTheFlyServerOptions Options)
+{
+	return MakeUnique<FCookOnTheFlyConnectionServer>(MoveTemp(Options));
+}
+
+}} // namespace UE::Cook
diff --git a/Engine/Source/Runtime/CookOnTheFly/Private/CookOnTheFlyModule.cpp b/Engine/Source/Runtime/CookOnTheFly/Private/CookOnTheFlyModule.cpp
new file mode 100644
--- /dev/null
+++ b/Engine/Source/Runtime/CookOnTheFly/Private/CookOnTheFlyModule.cpp
@@ -0,0 +1,69 @@
+// Copyright Epic Games, Inc. All Rights Reserved.
+
+#include "CookOnTheFly.h"
+#include "Modules/ModuleManager.h"
+
+DEFINE_LOG_CATEGORY(LogCookOnTheFly);
+
+namespace UE { namespace Cook
+{
+
+TUniquePtr<ICookOnTheFlyConnectionServer> MakeCookOnTheFlyConnectionServer(FCookOnTheFlyServerOptions Options);
+	
+TUniquePtr<ICookOnTheFlyServerConnection> MakeServerConnection();
+
+}}
+
+class FCookOnTheFlyModule final
+	: public UE::Cook::ICookOnTheFlyModule
+{
+public:
+
+	//
+	// IModuleInterface interface
+	//
+
+	virtual void StartupModule( ) override { }
+
+	virtual void ShutdownModule( ) override
+	{ 
+		if (ServerConnection.IsValid())
+		{
+			ServerConnection->Disconnect();
+			ServerConnection.Reset();
+		}
+	}
+
+	virtual bool SupportsDynamicReloading( ) override { return false; }
+
+	//
+	// ICookOnTheFlyModule interface
+	//
+
+	virtual TUniquePtr<UE::Cook::ICookOnTheFlyConnectionServer> CreateConnectionServer(UE::Cook::FCookOnTheFlyServerOptions Options) override
+	{
+		return UE::Cook::MakeCookOnTheFlyConnectionServer(Options);
+	}
+
+	virtual bool ConnectToServer(const UE::Cook::FCookOnTheFlyHostOptions& HostOptions) override
+	{
+		if (!ServerConnection.IsValid())
+		{
+			ServerConnection = UE::Cook::MakeServerConnection();
+			return ServerConnection->Connect(HostOptions);
+		}
+
+		return ServerConnection->IsConnected();
+	}
+
+	virtual UE::Cook::ICookOnTheFlyServerConnection& GetServerConnection() override
+	{
+		checkf(ServerConnection.IsValid(), TEXT("Call connect before accessing the server connection"));
+		return *ServerConnection;
+	}
+
+private:
+	TUniquePtr<UE::Cook::ICookOnTheFlyServerConnection> ServerConnection;
+};
+
+IMPLEMENT_MODULE(FCookOnTheFlyModule, CookOnTheFly);
diff --git a/Engine/Source/Runtime/CookOnTheFly/Private/CookOnTheFlyServerConnection.cpp b/Engine/Source/Runtime/CookOnTheFly/Private/CookOnTheFlyServerConnection.cpp
new file mode 100644
--- /dev/null
+++ b/Engine/Source/Runtime/CookOnTheFly/Private/CookOnTheFlyServerConnection.cpp
@@ -0,0 +1,382 @@
+// Copyright Epic Games, Inc. All Rights Reserved.
+
+#include "CookOnTheFly.h"
+#include "HAL/PlatformMisc.h"
+#include "Async/Async.h"
+#include "IPAddress.h"
+#include "SocketSubsystem.h"
+#include "MultichannelTcpSocket.h"
+#include "NetworkMessage.h"
+#include "Sockets.h"
+#include "Misc/App.h"
+#include "Misc/DateTime.h"
+#include "Serialization/MemoryReader.h"
+
+DEFINE_LOG_CATEGORY_STATIC(LogCotfServerConnection, Log, All);
+
+static TArray<TSharedPtr<FInternetAddr>> GetAddressFromString(ISocketSubsystem& SocketSubsystem, TArrayView<const FString> HostAddresses, const int32 Port)
+{
+	TArray<TSharedPtr<FInternetAddr>> InterntAddresses;
+
+	for (const FString& HostAddr : HostAddresses)
+	{
+		TSharedPtr<FInternetAddr> Addr = SocketSubsystem.GetAddressFromString(HostAddr);
+
+		if (!Addr.IsValid() || !Addr->IsValid())
+		{
+			FAddressInfoResult GAIRequest = SocketSubsystem.GetAddressInfo(*HostAddr, nullptr, EAddressInfoFlags::Default, NAME_None);
+			if (GAIRequest.ReturnCode == SE_NO_ERROR && GAIRequest.Results.Num() > 0)
+			{
+				Addr = GAIRequest.Results[0].Address;
+			}
+		}
+
+		if (Addr.IsValid() && Addr->IsValid())
+		{
+			Addr->SetPort(Port);
+			InterntAddresses.Emplace(MoveTemp(Addr));
+		}
+	}
+
+	return InterntAddresses;
+}
+
+class FCookOnTheFlyServerConnection
+	: public UE::Cook::ICookOnTheFlyServerConnection
+{
+public:
+	FCookOnTheFlyServerConnection()
+	{ }
+
+	~FCookOnTheFlyServerConnection()
+	{
+		Disconnect();
+	}
+
+private:
+
+	virtual bool Connect(const UE::Cook::FCookOnTheFlyHostOptions& HostOptions) override
+	{
+		check(HostOptions.Hosts.Num());
+
+		const int32 Port = HostOptions.Port > 0 ? HostOptions.Port : UE::Cook::DefaultCookOnTheFlyServingPort;
+
+		ISocketSubsystem& SocketSubsystem = *ISocketSubsystem::Get();
+		TArray<TSharedPtr<FInternetAddr>> HostAddresses = GetAddressFromString(SocketSubsystem, HostOptions.Hosts, Port);
+
+		if (!HostAddresses.Num())
+		{
+			UE_LOG(LogCotfServerConnection, Error, TEXT("No valid COTF server address found"));
+			return false;
+		}
+
+		bool bConnected = false;
+		const double ServerWaitEndTime = FPlatformTime::Seconds() + HostOptions.ServerStartupWaitTime.GetTotalSeconds();
+
+		for (;;)
+		{
+			for (const TSharedPtr<FInternetAddr>& Addr : HostAddresses)
+			{
+				UE_LOG(LogCotfServerConnection, Display, TEXT("Connecting to COTF server at '%s'..."), *Addr->ToString(true));
+			
+				Socket.Reset(SocketSubsystem.CreateSocket(NAME_Stream, TEXT("COTF-ServerConnection"), Addr->GetProtocolType()));
+				if (Socket.IsValid() && Socket->Connect(*Addr))
+				{
+					ServerAddr = Addr;
+					bConnected = true;
+					break;
+				}
+			}
+
+			if (bConnected || FPlatformTime::Seconds() > ServerWaitEndTime)
+			{
+				break;
+			}
+
+			FPlatformProcess::Sleep(1.0f);
+		};
+
+		if (!bConnected)
+		{
+			UE_LOG(LogCotfServerConnection, Error, TEXT("Failed to connect to COTF server"));
+			return false;
+		}
+
+		if (!SendHandshakeMessage())
+		{
+			UE_LOG(LogCotfServerConnection, Fatal, TEXT("Failed to handshake with COTF server at '%s'"), *ServerAddr->ToString(true));
+			return false;
+		}
+
+		Thread = AsyncThread([this] { return ThreadEntry(); },  8 * 1024, TPri_Normal);
+
+		UE_LOG(LogCotfServerConnection, Log, TEXT("Connected to COTF server at '%s'"), *ServerAddr->ToString(true));
+
+		return true;
+	}
+
+	virtual bool IsConnected() const override
+	{
+		return ClientId > 0;
+	}
+
+	virtual TFuture<UE::Cook::FCookOnTheFlyResponse> SendRequest(const UE::Cook::FCookOnTheFlyRequest& Request) override
+	{
+		using namespace UE::Cook;
+
+		const uint32 CorrelationId					= NextCorrelationId++;
+		FCookOnTheFlyMessageHeader RequestHeader	= Request.GetHeader();
+
+		RequestHeader.MessageType	= RequestHeader.MessageType | ECookOnTheFlyMessage::Request;
+		RequestHeader.MessageStatus = ECookOnTheFlyMessageStatus::Ok;
+		RequestHeader.SenderId		= ClientId;
+		RequestHeader.CorrelationId = CorrelationId;
+		RequestHeader.Timestamp		= FDateTime::UtcNow().GetTicks();
+
+		FBufferArchive RequestPayload;
+		RequestPayload.Reserve(Request.TotalSize());
+
+		RequestPayload << RequestHeader;
+		RequestPayload << const_cast<TArray<uint8>&>(Request.GetBody());
+
+		FPendingRequest* PendingRequest = Alloc(CorrelationId);
+		PendingRequest->RequestHeader = RequestHeader;
+		
+		TFuture<FCookOnTheFlyResponse> FutureResponse = PendingRequest->ResponsePromise.GetFuture();
+
+		UE_LOG(LogCotfServerConnection, Verbose, TEXT("Sending: %s, Size='%lld'"), *RequestHeader.ToString(), Request.TotalSize());
+
+		if (SendMessage(RequestPayload))
+		{
+			return FutureResponse;
+		}
+		else
+		{
+			UE_LOG(LogCotfServerConnection, Warning, TEXT("Failed to send: %s, Size='%lld'"), *RequestHeader.ToString(), Request.TotalSize());
+
+			Free(PendingRequest);
+
+			FCookOnTheFlyResponse ErrorResponse;
+			ErrorResponse.SetStatus(ECookOnTheFlyMessageStatus::Error);
+
+			TPromise<FCookOnTheFlyResponse> ErrorResponsePromise;
+			ErrorResponsePromise.SetValue(ErrorResponse);
+			return ErrorResponsePromise.GetFuture();
+		}
+	}
+
+	virtual void Disconnect() override
+	{
+		if (IsConnected())
+		{
+			if (!bStopRequested.Exchange(true))
+			{
+				Socket->Close();
+				Thread.Wait();
+				Thread.Reset();
+				Socket.Reset();
+				ClientId  = -1;
+			}
+		}
+	}
+
+	DECLARE_DERIVED_EVENT(FCookOnTheFlyServerConnection, UE::Cook::ICookOnTheFlyServerConnection::FMessageEvent, FMessageEvent);
+	virtual FMessageEvent& OnMessage() override
+	{
+		return MessageEvent;
+	}
+
+	bool SendMessage(const TArray<uint8>& Message)
+	{
+		if (!FNFSMessageHeader::WrapAndSendPayload(Message, FSimpleAbstractSocket_FSocket(Socket.Get())))
+		{
+			UE_LOG(LogCotfServerConnection, Fatal, TEXT("Failed sending payload to COTF server"));
+			return false;
+		}
+
+		return true;
+	}
+
+	bool ReceiveMessage(FArrayReader& Message)
+	{
+		if (!FNFSMessageHeader::ReceivePayload(Message, FSimpleAbstractSocket_FSocket(Socket.Get())))
+		{
+			UE_LOG(LogCotfServerConnection, Warning, TEXT("Failed reveiving payload from COTF server"));
+			return false;
+		}
+
+		return true;
+	}
+
+	bool SendHandshakeMessage()
+	{
+		using namespace UE::Cook;
+
+		FCookOnTheFlyMessage HandshakeRequest(ECookOnTheFlyMessage::Handshake | ECookOnTheFlyMessage::Request);
+		{
+			FString PlatformName(ANSI_TO_TCHAR(FPlatformProperties::PlatformName()));
+			FString ProjectName(FApp::GetProjectName());
+
+			TUniquePtr<FArchive> Ar = HandshakeRequest.WriteBody();
+			*Ar << PlatformName;
+			*Ar << ProjectName;
+		}
+
+		FBufferArchive HandshakeRequestPayload;
+		HandshakeRequestPayload << HandshakeRequest;
+
+		if (!SendMessage(HandshakeRequestPayload))
+		{
+			return false;
+		}
+
+		FArrayReader HandshakeResponsePayload;
+		if (!ReceiveMessage(HandshakeResponsePayload))
+		{
+			return false;
+		}
+
+		FCookOnTheFlyMessage HandshakeResponse;
+		HandshakeResponsePayload << HandshakeResponse;
+		{
+			TUniquePtr<FArchive> Ar = HandshakeResponse.ReadBody();
+			*Ar << ClientId;
+		}
+
+		UE_CLOG(ClientId > 0, LogCotfServerConnection, Display, TEXT("Connected to server with ID='%d'"), ClientId);
+
+		return ClientId > 0;
+	}
+
+	void ThreadEntry()
+	{
+		using namespace UE::Cook;
+
+		while (!bStopRequested.Load())
+		{
+			FArrayReader MessagePayload;
+			if (!ReceiveMessage(MessagePayload))
+			{
+				UE_LOG(LogCotfServerConnection, Warning, TEXT("Failed to receive message from '%s'"), *ServerAddr->ToString(true));
+				break;
+			}
+
+			FCookOnTheFlyMessageHeader MessageHeader;
+			TArray<uint8> MessageBody;
+
+			MessagePayload << MessageHeader;
+			MessagePayload << MessageBody;
+
+			UE_LOG(LogCotfServerConnection, Verbose, TEXT("Received: %s, Size='%lld'"), *MessageHeader.ToString(), MessagePayload.Num());
+
+			const bool bIsResponse = EnumHasAnyFlags(MessageHeader.MessageType, ECookOnTheFlyMessage::Response);
+			const bool bIsRequest = EnumHasAnyFlags(MessageHeader.MessageType, ECookOnTheFlyMessage::Request);
+			EnumRemoveFlags(MessageHeader.MessageType, ECookOnTheFlyMessage::TypeFlags);
+
+			if (bIsRequest)
+			{
+				UE_CLOG(
+					MessageHeader.MessageType != ECookOnTheFlyMessage::Heartbeat,
+					LogCotfServerConnection, Fatal, TEXT("Invalid server request message '%s'"), LexToString(MessageHeader.MessageType));
+
+				FCookOnTheFlyMessage HeartbeatResponse(ECookOnTheFlyMessage::Heartbeat | ECookOnTheFlyMessage::Response);
+				FCookOnTheFlyMessageHeader& ResponseHeader = HeartbeatResponse.GetHeader();
+
+				ResponseHeader.MessageStatus	= ECookOnTheFlyMessageStatus::Ok;
+				ResponseHeader.SenderId			= ClientId;
+				ResponseHeader.CorrelationId	= MessageHeader.CorrelationId;
+				ResponseHeader.Timestamp		= FDateTime::UtcNow().GetTicks();
+				
+				FBufferArchive ResponsePayload;
+				ResponsePayload << HeartbeatResponse;
+
+				UE_LOG(LogCotfServerConnection, Warning, TEXT("Sending heartbeat response to '%s'"), *ServerAddr->ToString(true));
+
+				if (!SendMessage(ResponsePayload))
+				{
+					UE_LOG(LogCotfServerConnection, Warning, TEXT("Failed to send heartbeat response to '%s'"), *ServerAddr->ToString(true));
+					break;
+				}
+			}
+			else if (bIsResponse)
+			{
+				FPendingRequest* PendingRequest = GetRequest(MessageHeader.CorrelationId);
+				check(PendingRequest);
+				check(PendingRequest->RequestHeader.CorrelationId == MessageHeader.CorrelationId);
+
+				FCookOnTheFlyResponse Response;
+				Response.SetHeader(MessageHeader);
+				Response.SetBody(MoveTemp(MessageBody));
+
+				PendingRequest->ResponsePromise.SetValue(MoveTemp(Response));
+
+				Free(PendingRequest);
+			}
+			else
+			{
+				FCookOnTheFlyMessage Message;
+				Message.SetHeader(MessageHeader);
+				Message.SetBody(MoveTemp(MessageBody));
+				
+				if (MessageEvent.IsBound())
+				{
+					MessageEvent.Broadcast(Message);
+				}
+			}
+		}
+
+		UE_LOG(LogCotfServerConnection, Display, TEXT("Terminating connection to server '%s'"), *ServerAddr->ToString(true));
+	}
+
+	struct FPendingRequest
+	{
+		UE::Cook::FCookOnTheFlyMessageHeader RequestHeader;
+		TPromise<UE::Cook::FCookOnTheFlyResponse> ResponsePromise;
+	};
+
+	FPendingRequest* Alloc(uint32 CorrelationId)
+	{
+		FScopeLock _(&RequestsCriticalSection);
+		TUniquePtr<FPendingRequest>& NewPendingRequest = PendingRequests.Add(CorrelationId, MakeUnique<FPendingRequest>());
+		return NewPendingRequest.Get();
+	}
+
+	void Free(FPendingRequest* PendingRequest)
+	{
+		FScopeLock _(&RequestsCriticalSection);
+		PendingRequests.Remove(PendingRequest->RequestHeader.CorrelationId);
+	}
+	
+	FPendingRequest* GetRequest(uint32 CorrelationId)
+	{
+		FScopeLock _(&RequestsCriticalSection);
+		
+		if (TUniquePtr<FPendingRequest>* PendingRequest = PendingRequests.Find(CorrelationId))
+		{
+			return PendingRequest->Get();
+		}
+
+		return nullptr;
+	}
+
+	FMessageEvent MessageEvent;
+	TSharedPtr<FInternetAddr> ServerAddr;
+	TUniquePtr<FSocket> Socket;
+	uint32 ClientId = 0;
+	TFuture<void> Thread;
+	TAtomic<bool> bStopRequested { false };
+
+	FCriticalSection RequestsCriticalSection;
+	TMap<uint32, TUniquePtr<FPendingRequest>> PendingRequests;
+	TAtomic<uint32> NextCorrelationId { 1 };
+};
+
+namespace UE { namespace Cook
+{
+
+TUniquePtr<ICookOnTheFlyServerConnection> MakeServerConnection()
+{
+	return MakeUnique<FCookOnTheFlyServerConnection>();
+}
+
+}} // namespace UE::Cook
diff --git a/Engine/Source/Runtime/Core/Private/GenericPlatform/GenericPlatformIoDispatcher.cpp b/Engine/Source/Runtime/Core/Private/GenericPlatform/GenericPlatformIoDispatcher.cpp
--- a/Engine/Source/Runtime/Core/Private/GenericPlatform/GenericPlatformIoDispatcher.cpp
+++ b/Engine/Source/Runtime/Core/Private/GenericPlatform/GenericPlatformIoDispatcher.cpp
@@ -5,7 +5,6 @@
 #include "Async/AsyncFileHandle.h"
 #include "HAL/Event.h"
 #include "HAL/PlatformFileManager.h"
-#include "ProfilingDebugging/CountersTrace.h"
 #include "HAL/RunnableThread.h"
 #include "HAL/IConsoleManager.h"
 #include "Misc/ScopeLock.h"
@@ -13,6 +12,14 @@
 //PRAGMA_DISABLE_OPTIMIZATION
 
 
+TRACE_DECLARE_INT_COUNTER_EXTERN(IoDispatcherFileBackendSequentialReads);
+TRACE_DECLARE_INT_COUNTER_EXTERN(IoDispatcherFileBackendForwardSeeks);
+TRACE_DECLARE_INT_COUNTER_EXTERN(IoDispatcherFileBackendBackwardSeeks);
+TRACE_DECLARE_INT_COUNTER_EXTERN(IoDispatcherFileBackendSwitchContainerSeeks);
+TRACE_DECLARE_MEMORY_COUNTER_EXTERN(IoDispatcherFileBackendTotalSeekDistance);
+TRACE_DECLARE_INT_COUNTER_EXTERN(IoDispatcherFileBackendFileSystemRequests);
+TRACE_DECLARE_MEMORY_COUNTER_EXTERN(IoDispatcherFileBackendFileSystemTotalBytesRead);
+
 FGenericFileIoStoreEventQueue::FGenericFileIoStoreEventQueue()
 	: ServiceEvent(FPlatformProcess::GetSynchEventFromPool())
 {
@@ -113,6 +120,7 @@ bool FGenericFileIoStoreImpl::StartRequests(FFileIoStoreRequestQueue& RequestQue
 			TRACE_CPUPROFILER_EVENT_SCOPE(ReadBlockFromFile);
 			NextRequest->bFailed = true;
 			int32 RetryCount = 0;
+			TRACE_COUNTER_INCREMENT(IoDispatcherFileBackendFileSystemRequests);
 			while (RetryCount++ < 10)
 			{
 				if (!FileHandle->Seek(NextRequest->Offset))
@@ -126,6 +134,7 @@ bool FGenericFileIoStoreImpl::StartRequests(FFileIoStoreRequestQueue& RequestQue
 					continue;
 				}
 				NextRequest->bFailed = false;
+				TRACE_COUNTER_ADD(IoDispatcherFileBackendFileSystemTotalBytesRead, NextRequest->Size);
 				BlockCache.Store(NextRequest);
 				break;
 			}
diff --git a/Engine/Source/Runtime/Core/Private/IO/IoDispatcherFileBackend.cpp b/Engine/Source/Runtime/Core/Private/IO/IoDispatcherFileBackend.cpp
--- a/Engine/Source/Runtime/Core/Private/IO/IoDispatcherFileBackend.cpp
+++ b/Engine/Source/Runtime/Core/Private/IO/IoDispatcherFileBackend.cpp
@@ -21,6 +21,19 @@
 TRACE_DECLARE_INT_COUNTER(IoDispatcherLatencyCircuitBreaks, TEXT("IoDispatcher/LatencyCircuitBreaks"));
 TRACE_DECLARE_INT_COUNTER(IoDispatcherSeekDistanceCircuitBreaks, TEXT("IoDispatcher/SeekDistanceCircuitBreaks"));
 TRACE_DECLARE_INT_COUNTER(IoDispatcherNumPriorityQueues, TEXT("IoDispatcher/NumPriorityQueues"));
+TRACE_DECLARE_MEMORY_COUNTER(IoDispatcherFileBackendTotalBytesRead, TEXT("IoDispatcherFileBackend/TotalBytesRead"));
+TRACE_DECLARE_MEMORY_COUNTER(IoDispatcherFileBackendFileSystemTotalBytesRead, TEXT("IoDispatcherFileBackend/TotalBytesReadFromFileSystem"));
+TRACE_DECLARE_MEMORY_COUNTER(IoDispatcherFileBackendTotalBytesScattered, TEXT("IoDispatcherFileBackend/TotalBytesScattered"));
+TRACE_DECLARE_INT_COUNTER(IoDispatcherFileBackendFileSystemRequests, TEXT("IoDispatcherFileBackend/FileSystemRequests"));
+TRACE_DECLARE_INT_COUNTER(IoDispatcherFileBackendCacheHits, TEXT("IoDispatcherFileBackend/CacheHits"));
+TRACE_DECLARE_INT_COUNTER(IoDispatcherFileBackendCacheMisses, TEXT("IoDispatcherFileBackend/CacheMisses"));
+TRACE_DECLARE_MEMORY_COUNTER(IoDispatcherFileBackendTotalBytesCached, TEXT("IoDispatcherFileBackend/CacheTotalBytesStored"));
+TRACE_DECLARE_INT_COUNTER(IoDispatcherFileBackendQueueLength, TEXT("IoDispatcherFileBackend/QueueLength"));
+TRACE_DECLARE_INT_COUNTER(IoDispatcherFileBackendSequentialReads, TEXT("IoDispatcherFileBackend/SequentialReads"));
+TRACE_DECLARE_INT_COUNTER(IoDispatcherFileBackendForwardSeeks, TEXT("IoDispatcherFileBackend/ForwardSeeks"));
+TRACE_DECLARE_INT_COUNTER(IoDispatcherFileBackendBackwardSeeks, TEXT("IoDispatcherFileBackend/BackwardSeeks"));
+TRACE_DECLARE_INT_COUNTER(IoDispatcherFileBackendSwitchContainerSeeks, TEXT("IoDispatcherFileBackend/SwitchContainerSeeks"));
+TRACE_DECLARE_MEMORY_COUNTER(IoDispatcherFileBackendTotalSeekDistance, TEXT("IoDispatcherFileBackend/TotalSeekDistance"));
 
 //PRAGMA_DISABLE_OPTIMIZATION
 
@@ -197,37 +210,30 @@ bool FFileIoStoreBlockCache::Read(FFileIoStoreReadRequest* Block)
 		return false;
 	}
 	check(Block->Buffer);
-	FCachedBlock* CachedBlock = nullptr;
+	FCachedBlock* CachedBlock = CachedBlocks.FindRef(Block->Key.Hash);
+	if (CachedBlock)
 	{
-		FScopeLock Lock(&CriticalSection);
-		CachedBlock = CachedBlocks.FindRef(Block->Key.Hash);
-		if (CachedBlock)
-		{
-			CachedBlock->bLocked = true;
-
-			CachedBlock->LruPrev->LruNext = CachedBlock->LruNext;
-			CachedBlock->LruNext->LruPrev = CachedBlock->LruPrev;
+		CachedBlock->LruPrev->LruNext = CachedBlock->LruNext;
+		CachedBlock->LruNext->LruPrev = CachedBlock->LruPrev;
 
-			CachedBlock->LruPrev = &CacheLruHead;
-			CachedBlock->LruNext = CacheLruHead.LruNext;
+		CachedBlock->LruPrev = &CacheLruHead;
+		CachedBlock->LruNext = CacheLruHead.LruNext;
 
-			CachedBlock->LruPrev->LruNext = CachedBlock;
-			CachedBlock->LruNext->LruPrev = CachedBlock;
-		}
+		CachedBlock->LruPrev->LruNext = CachedBlock;
+		CachedBlock->LruNext->LruPrev = CachedBlock;
 	}
 
 	if (!CachedBlock)
 	{
 		FFileIoStats::OnBlockCacheMiss(ReadBufferSize);
+		TRACE_COUNTER_INCREMENT(IoDispatcherFileBackendCacheMisses);
 		return false;
 	}
 	check(CachedBlock->Buffer);
 	FFileIoStats::OnBlockCacheHit(ReadBufferSize);
 	FMemory::Memcpy(Block->Buffer->Memory, CachedBlock->Buffer, ReadBufferSize);
-	{
-		FScopeLock Lock(&CriticalSection);
-		CachedBlock->bLocked = false;
-	}
+	TRACE_COUNTER_INCREMENT(IoDispatcherFileBackendCacheHits);
+
 	return true;
 }
 
@@ -240,40 +246,29 @@ void FFileIoStoreBlockCache::Store(const FFileIoStoreReadRequest* Block)
 	}
 	check(Block->Buffer);
 	check(Block->Buffer->Memory);
-	FCachedBlock* BlockToReplace = nullptr;
+	FCachedBlock* BlockToReplace = CacheLruTail.LruPrev;
+	if (BlockToReplace == &CacheLruHead)
 	{
-		FScopeLock Lock(&CriticalSection);
-		BlockToReplace = CacheLruTail.LruPrev;
-		while (BlockToReplace != &CacheLruHead && BlockToReplace->bLocked)
-		{
-			BlockToReplace = BlockToReplace->LruPrev;
-		}
-		if (BlockToReplace == &CacheLruHead)
-		{
-			return;
-		}
-		CachedBlocks.Remove(BlockToReplace->Key);
-		BlockToReplace->bLocked = true;
-		BlockToReplace->Key = Block->Key.Hash;
+		return;
+	}
+	check(BlockToReplace);
+	CachedBlocks.Remove(BlockToReplace->Key);
+	BlockToReplace->Key = Block->Key.Hash;
 
-		BlockToReplace->LruPrev->LruNext = BlockToReplace->LruNext;
-		BlockToReplace->LruNext->LruPrev = BlockToReplace->LruPrev;
+	BlockToReplace->LruPrev->LruNext = BlockToReplace->LruNext;
+	BlockToReplace->LruNext->LruPrev = BlockToReplace->LruPrev;
 
-		BlockToReplace->LruPrev = &CacheLruHead;
-		BlockToReplace->LruNext = CacheLruHead.LruNext;
+	BlockToReplace->LruPrev = &CacheLruHead;
+	BlockToReplace->LruNext = CacheLruHead.LruNext;
 
-		BlockToReplace->LruPrev->LruNext = BlockToReplace;
-		BlockToReplace->LruNext->LruPrev = BlockToReplace;
-	}
-	check(BlockToReplace);
+	BlockToReplace->LruPrev->LruNext = BlockToReplace;
+	BlockToReplace->LruNext->LruPrev = BlockToReplace;
+	
 	check(BlockToReplace->Buffer);
 	FMemory::Memcpy(BlockToReplace->Buffer, Block->Buffer->Memory, ReadBufferSize);
 	FFileIoStats::OnBlockCacheStore(ReadBufferSize);
-	{
-		FScopeLock Lock(&CriticalSection);
-		BlockToReplace->bLocked = false;
-		CachedBlocks.Add(BlockToReplace->Key, BlockToReplace);
-	}
+	CachedBlocks.Add(BlockToReplace->Key, BlockToReplace);
+	TRACE_COUNTER_ADD(IoDispatcherFileBackendTotalBytesCached, ReadBufferSize);
 }
 
 bool FFileIoStoreOffsetSortedRequestQueue::RequestSortPredicate(const FFileIoStoreReadRequestSortKey& A, const FFileIoStoreReadRequestSortKey& B)
@@ -516,6 +511,7 @@ FFileIoStoreReadRequest* FFileIoStoreRequestQueue::Pop()
 	
 	check(Result->QueueStatus == FFileIoStoreReadRequest::QueueStatus_InQueue);
 	Result->QueueStatus = FFileIoStoreReadRequest::QueueStatus_Started;
+	TRACE_COUNTER_SET(IoDispatcherFileBackendQueueLength, Heap.Num());
 	return Result;
 }
 
@@ -549,6 +545,7 @@ void FFileIoStoreRequestQueue::Push(FFileIoStoreReadRequest& Request)
 	else
 	{
 		Heap.HeapPush(&Request, QueueSortFunc);
+		TRACE_COUNTER_SET(IoDispatcherFileBackendQueueLength, Heap.Num());
 	}
 }
 
@@ -573,6 +570,7 @@ void FFileIoStoreRequestQueue::Push(FFileIoStoreReadRequestList& Requests)
 		}
 	}
 
+	TRACE_COUNTER_SET(IoDispatcherFileBackendQueueLength, Heap.Num());
 }
 
 void FFileIoStoreRequestQueue::UpdateOrder()
@@ -1395,6 +1393,7 @@ void FFileIoStore::FinalizeCompressedBlock(FFileIoStoreCompressedBlock* Compress
 	{
 		FFileIoStoreBlockScatter& Scatter = CompressedBlock->ScatterList[ScatterIndex];
 		FFileIoStats::OnBytesScattered(Scatter.Size);
+		TRACE_COUNTER_ADD(IoDispatcherFileBackendTotalBytesScattered, Scatter.Size);
 		Scatter.Request->bFailed |= CompressedBlock->bFailed;
 		check(!CompressedBlock->bCancelled || !Scatter.Request->DispatcherRequest || Scatter.Request->DispatcherRequest->IsCancelled());
 		check(Scatter.Request->UnfinishedReadsCount > 0);
@@ -1426,6 +1425,7 @@ FIoRequestImpl* FFileIoStore::GetCompletedRequests()
 		FFileIoStoreReadRequest* CompletedRequest = *It;
 
 		FFileIoStats::OnReadComplete(CompletedRequest->Size);
+		TRACE_COUNTER_ADD(IoDispatcherFileBackendTotalBytesRead, CompletedRequest->Size);
 
 		if (!CompletedRequest->ImmediateScatter.Request)
 		{
@@ -1501,6 +1501,8 @@ FIoRequestImpl* FFileIoStore::GetCompletedRequests()
 		else
 		{
 			FFileIoStats::OnBytesScattered(CompletedRequest->ImmediateScatter.Size);
+			TRACE_COUNTER_ADD(IoDispatcherFileBackendTotalBytesScattered, CompletedRequest->ImmediateScatter.Size);
+
 
 			check(!CompletedRequest->Buffer);
 			FFileIoStoreResolvedRequest* CompletedResolvedRequest = CompletedRequest->ImmediateScatter.Request;
@@ -1856,22 +1858,15 @@ TRACE_DECLARE_MEMORY_COUNTER(IoDispatcherFileBackendQueuedFilesystemReadBytes,	T
 // These are long term totals that grow forever and may be less useful?
 
 // IoDispatcher thread 
-// Total bytes read returned from backend whether read from disk or returned from cache 
-TRACE_DECLARE_MEMORY_COUNTER(IoDispatcherFileBackendTotalBytesRead,				TEXT("IoDispatcherFileBackend/TotalBytesRead"));			
 // Total bytes passing through the decompression stage, even if they were not compressed
 TRACE_DECLARE_MEMORY_COUNTER(IoDispatcherFileBackendTotalBytesUncompressedIn,	TEXT("IoDispatcherFileBackend/TotalBytesUncompressedIn"));
 // Total bytes passing through the decompression stage, even if they were not compressed
 TRACE_DECLARE_MEMORY_COUNTER(IoDispatcherFileBackendTotalBytesUncompressedOut,	TEXT("IoDispatcherFileBackend/TotalBytesUncompressedOut"));
-// Total bytes provided into IoBuffers for users 
-TRACE_DECLARE_MEMORY_COUNTER(IoDispatcherFileBackendTotalBytesScattered,		TEXT("IoDispatcherFileBackend/TotalBytesScattered"));
-
 // FileIoStore thread
 // Total number of reads executed through filesystem API layer (useful for calculating ratios of seeks vs sequential reads)
 TRACE_DECLARE_MEMORY_COUNTER(IoDispatcherFileBackendTotalFilesystemReads,		TEXT("IoDispatcherFileBackend/TotalFilesystemReads"));		
 // Total bytes read from platform uncompressed or decompressed for us below the API layer
 TRACE_DECLARE_MEMORY_COUNTER(IoDispatcherFileBackendTotalFilesystemBytesRead,	TEXT("IoDispatcherFileBackend/TotalFilesystemBytesRead"));
-// Total distance in bytes from seeks backwards and forwards
-TRACE_DECLARE_MEMORY_COUNTER(IoDispatcherFileBackendTotalSeekDistance,			TEXT("IoDispatcherFileBackend/TotalSeekDistance"));
 // Total bytes stored in block cache including replacements
 TRACE_DECLARE_MEMORY_COUNTER(IoDispatcherFileBackendTotalBlockCacheBytesStored,	TEXT("IoDispatcherFileBackend/TotalBlockCacheBytesStored"));
 // Total bytes retrieved from block cache
diff --git a/Engine/Source/Runtime/Core/Private/IO/IoDispatcherFileBackendTypes.h b/Engine/Source/Runtime/Core/Private/IO/IoDispatcherFileBackendTypes.h
--- a/Engine/Source/Runtime/Core/Private/IO/IoDispatcherFileBackendTypes.h
+++ b/Engine/Source/Runtime/Core/Private/IO/IoDispatcherFileBackendTypes.h
@@ -539,10 +539,8 @@ private:
 		FCachedBlock* LruNext = nullptr;
 		uint64 Key = 0;
 		uint8* Buffer = nullptr;
-		bool bLocked = false;
 	};
 
-	FCriticalSection CriticalSection;
 	uint8* CacheMemory = nullptr;
 	TMap<uint64, FCachedBlock*> CachedBlocks;
 	FCachedBlock CacheLruHead;
diff --git a/Engine/Source/Runtime/Core/Private/IO/IoStore.cpp b/Engine/Source/Runtime/Core/Private/IO/IoStore.cpp
--- a/Engine/Source/Runtime/Core/Private/IO/IoStore.cpp
+++ b/Engine/Source/Runtime/Core/Private/IO/IoStore.cpp
@@ -1085,6 +1085,18 @@ private:
 	void WriteEntry(FIoStoreWriteQueueEntry* Entry)
 	{
 		TRACE_CPUPROFILER_EVENT_SCOPE(WriteEntry);
+		const int32* FindExistingIndex = Toc.GetTocEntryIndex(Entry->ChunkId);
+		if (FindExistingIndex)
+		{
+			checkf(Toc.GetTocResource().ChunkMetas[*FindExistingIndex].ChunkHash == Entry->ChunkHash, TEXT("Chunk id has already been added with different content"));
+			for (FChunkBlock& ChunkBlock : Entry->ChunkBlocks)
+			{
+				WriterContext->FreeCompressionBuffer(ChunkBlock.IoBuffer, Entry->ChunkBlocks.Num());
+				ChunkBlock.IoBuffer = nullptr;
+			}
+			return;
+		}
+
 		FPartition* TargetPartition = &Partitions[CurrentPartitionIndex];
 		int32 NextPartitionIndexToTry = CurrentPartitionIndex + 1;
 		if (Entry->PartitionIndex >= 0)
@@ -1407,6 +1419,36 @@ TIoStatusOr<FIoStoreWriterResult> FIoStoreWriter::Flush()
 	return Impl->Flush();
 }
 
+class FIoStoreAsyncReadRequest::FImpl
+{
+public:
+	FImpl() = default;
+	
+	~FImpl()
+	{
+		Release();
+	}
+
+	void Release()
+	{
+		FScopeLock _(&CriticalSection);
+		for (TUniquePtr<IAsyncReadRequest>& Request : ReadRequests)
+		{
+			Request->WaitCompletion();
+		}
+		ReadRequests.Empty();
+		BlockBuffers.Empty();
+		Buffer = FIoBuffer();
+	}
+
+	FIoStoreAsyncReadRequest::FCallback Callback;
+	TArray<TUniquePtr<IAsyncReadRequest>> ReadRequests;
+	TArray<FIoBuffer> BlockBuffers;
+	FIoBuffer Buffer;
+	TAtomic<int32> NumRemainingBlocks {0};
+	FCriticalSection CriticalSection;
+};
+
 class FIoStoreReaderImpl
 {
 public:
@@ -1541,16 +1583,12 @@ public:
 			ResolvedSize = FMath::Min(Options.GetSize(), OffsetAndLength->GetLength() - RequestedOffset);
 		}
 
-		if (!ThreadBuffers)
-		{
-			ThreadBuffers = new FThreadBuffers();
-			ThreadBuffers->Register();
-		}
-		TArray<uint8>& CompressedBuffer = ThreadBuffers->CompressedBuffer;
-		TArray<uint8>& UncompressedBuffer = ThreadBuffers->UncompressedBuffer;
-
 		const FIoStoreTocResource& TocResource = Toc.GetTocResource();
 		const uint64 CompressionBlockSize = TocResource.Header.CompressionBlockSize;
+		TArray<uint8> CompressedBuffer;
+		CompressedBuffer.SetNumUninitialized(static_cast<int32>(CompressionBlockSize));
+		TArray<uint8> UncompressedBuffer;
+		UncompressedBuffer.SetNumUninitialized(static_cast<int32>(CompressionBlockSize));
 		FIoBuffer IoBuffer = FIoBuffer(ResolvedSize);
 		int32 FirstBlockIndex = int32(ResolvedOffset / CompressionBlockSize);
 		int32 LastBlockIndex = int32((Align(ResolvedOffset + ResolvedSize, CompressionBlockSize) - 1) / CompressionBlockSize);
@@ -1607,6 +1645,183 @@ public:
 		return IoBuffer;
 	}
 
+	TIoStatusOr<FIoStoreCompressedReadResult> ReadCompressed(const FIoChunkId& ChunkId, const FIoReadOptions& Options) const
+	{
+		TRACE_CPUPROFILER_EVENT_SCOPE(ReadChunkCompressed);
+
+		const FIoOffsetAndLength* OffsetAndLength = Toc.GetOffsetAndLength(ChunkId);
+		if (!OffsetAndLength)
+		{
+			return FIoStatus(EIoErrorCode::NotFound, TEXT("Unknown chunk ID"));
+		}
+
+		uint64 RequestedOffset = Options.GetOffset();
+		uint64 ResolvedOffset = OffsetAndLength->GetOffset() + RequestedOffset;
+		uint64 ResolvedSize = 0;
+		if (RequestedOffset <= OffsetAndLength->GetLength())
+		{
+			ResolvedSize = FMath::Min(Options.GetSize(), OffsetAndLength->GetLength() - RequestedOffset);
+		}
+
+		const FIoStoreTocResource& TocResource = Toc.GetTocResource();
+		const uint64 CompressionBlockSize = TocResource.Header.CompressionBlockSize;
+		TArray<uint8> CompressedBuffer;
+		CompressedBuffer.SetNumUninitialized(static_cast<int32>(CompressionBlockSize));
+		int32 FirstBlockIndex = int32(ResolvedOffset / CompressionBlockSize);
+		int32 LastBlockIndex = int32((Align(ResolvedOffset + ResolvedSize, CompressionBlockSize) - 1) / CompressionBlockSize);
+		uint64 TotalCompressedSize = 0;
+		for (int32 BlockIndex = FirstBlockIndex; BlockIndex <= LastBlockIndex; ++BlockIndex)
+		{
+			const FIoStoreTocCompressedBlockEntry& CompressionBlock = TocResource.CompressionBlocks[BlockIndex];
+			TotalCompressedSize += CompressionBlock.GetCompressedSize();
+		}
+		FIoStoreCompressedReadResult Result;
+		Result.IoBuffer = FIoBuffer(TotalCompressedSize);
+		Result.Blocks.Reserve(LastBlockIndex + 1 - FirstBlockIndex);
+		Result.UncompressedOffset = ResolvedOffset % CompressionBlockSize;
+		Result.UncompressedSize = ResolvedSize;
+		uint8* Dst = Result.IoBuffer.Data();
+		for (int32 BlockIndex = FirstBlockIndex; BlockIndex <= LastBlockIndex; ++BlockIndex)
+		{
+			const FIoStoreTocCompressedBlockEntry& CompressionBlock = TocResource.CompressionBlocks[BlockIndex];
+			uint32 RawSize = Align(CompressionBlock.GetCompressedSize(), FAES::AESBlockSize);
+			if (uint32(CompressedBuffer.Num()) < RawSize)
+			{
+				CompressedBuffer.SetNumUninitialized(RawSize);
+			}
+			
+			int32 PartitionIndex = int32(CompressionBlock.GetOffset() / TocResource.Header.PartitionSize);
+			int64 PartitionOffset = int64(CompressionBlock.GetOffset() % TocResource.Header.PartitionSize);
+			TUniquePtr<IAsyncReadRequest> ReadRequest(ContainerFileHandles[PartitionIndex]->ReadRequest(PartitionOffset, RawSize, AIOP_Normal, nullptr, CompressedBuffer.GetData()));
+			{
+				TRACE_CPUPROFILER_EVENT_SCOPE(WaitForIo);
+				ReadRequest->WaitCompletion();
+			}
+
+			if (EnumHasAnyFlags(TocResource.Header.ContainerFlags, EIoContainerFlags::Encrypted))
+			{
+				FAES::DecryptData(CompressedBuffer.GetData(), RawSize, DecryptionKey);
+			}
+			FIoStoreCompressedBlockInfo& BlockInfo = Result.Blocks.AddDefaulted_GetRef();
+			BlockInfo.CompressionMethod = TocResource.CompressionMethods[CompressionBlock.GetCompressionMethodIndex()];
+			BlockInfo.CompressedSize = CompressionBlock.GetCompressedSize();
+			BlockInfo.UncompressedSize = CompressionBlock.GetUncompressedSize();
+			FMemory::Memcpy(Dst, CompressedBuffer.GetData(), BlockInfo.CompressedSize);
+			Dst += BlockInfo.CompressedSize;
+		}
+		return Result;
+	}
+
+	FIoStoreAsyncReadRequest ReadAsync(const FIoChunkId& ChunkId, const FIoReadOptions& Options, FIoStoreAsyncReadRequest::FCallback&& Callback) const
+	{
+		FIoStoreAsyncReadRequest ReadRequest;
+		FIoStoreAsyncReadRequest::FImpl* Request  = ReadRequest.Impl.Get();
+
+		const FIoStoreTocResource& TocResource		= Toc.GetTocResource();
+		const FIoOffsetAndLength* OffsetAndLength	= Toc.GetOffsetAndLength(ChunkId);
+		const bool bIsEncrypted						= EnumHasAnyFlags(TocResource.Header.ContainerFlags, EIoContainerFlags::Encrypted);
+
+		if (!OffsetAndLength )
+		{
+			Callback(FIoStatus(EIoErrorCode::NotFound));
+			return ReadRequest;
+		}
+
+		const uint64 RequestedOffset	= Options.GetOffset();
+		const uint64 ResolvedOffset		= OffsetAndLength->GetOffset() + RequestedOffset;
+		uint64 ResolvedSize				= 0;
+
+		if (RequestedOffset <= OffsetAndLength->GetLength())
+		{
+			ResolvedSize = FMath::Min(Options.GetSize(), OffsetAndLength->GetLength() - RequestedOffset);
+		}
+
+		const uint64 CompressionBlockSize	= TocResource.Header.CompressionBlockSize;
+		const int32 FirstBlockIndex			= int32(ResolvedOffset / CompressionBlockSize);
+		const int32 LastBlockIndex			= int32((Align(ResolvedOffset + ResolvedSize, CompressionBlockSize) - 1) / CompressionBlockSize);
+		const int32	NumBlocks				= (LastBlockIndex - FirstBlockIndex) + 1;
+
+		Request->Callback = Callback;
+		Request->ReadRequests.SetNum(NumBlocks);
+		Request->BlockBuffers.SetNum(NumBlocks);
+		Request->Buffer = FIoBuffer(ResolvedSize);
+		Request->NumRemainingBlocks = NumBlocks;
+
+		for (int32 BlockIndex = FirstBlockIndex, Idx = 0; BlockIndex <= LastBlockIndex; ++BlockIndex, ++Idx)
+		{
+			const FIoStoreTocCompressedBlockEntry& CompressionBlock = TocResource.CompressionBlocks[BlockIndex];
+
+			FName CompressionMethod			= TocResource.CompressionMethods[CompressionBlock.GetCompressionMethodIndex()];
+			const uint32 CompressedSize		= CompressionBlock.GetCompressedSize();
+			const uint32 UncompressedSize	= CompressionBlock.GetUncompressedSize();
+			const uint32 RawSize			= Align(CompressedSize, FAES::AESBlockSize);
+			const int32 PartitionIndex		= int32(CompressionBlock.GetOffset() / TocResource.Header.PartitionSize);
+			const int64 PartitionOffset		= int64(CompressionBlock.GetOffset() % TocResource.Header.PartitionSize);
+
+			FIoBuffer& BlockBuffer = Request->BlockBuffers[Idx];
+			BlockBuffer = FIoBuffer(RawSize);
+
+			FAsyncFileCallBack FileReadCallback = [
+				this,
+				Request,
+				Idx,
+				CompressionMethod,
+				bIsEncrypted,
+				CompressedSize,
+				UncompressedSize,
+				CompressionBlockSize,
+				ResolvedOffset](bool bWasCancelled, IAsyncReadRequest* ReadRequest)
+			{
+				FIoBuffer& Buffer = Request->BlockBuffers[Idx];
+				check(ReadRequest->GetReadResults() == Buffer.Data());
+
+				if (bIsEncrypted)
+				{
+					FAES::DecryptData(Buffer.Data(), uint32(Buffer.DataSize()), DecryptionKey);
+				}
+
+				if (!CompressionMethod.IsNone())
+				{
+					FIoBuffer UncompresedBlockBuffer(UncompressedSize);
+
+					const bool bUncompressed = FCompression::UncompressMemory(
+						CompressionMethod,
+						UncompresedBlockBuffer.Data(),
+						int32(UncompressedSize),
+						Buffer.Data(), 
+						int32(CompressedSize));
+
+					check(bUncompressed);
+					Request->BlockBuffers[Idx] = UncompresedBlockBuffer;
+				}
+
+				if (--Request->NumRemainingBlocks == 0)
+				{
+					uint8* Dst = Request->Buffer.Data();
+					uint64 RemainingSize = Request->Buffer.DataSize();
+					uint64 OffsetInBlock = ResolvedOffset % CompressionBlockSize;
+
+					for (const FIoBuffer& Src : Request->BlockBuffers)
+					{
+						const uint64 SizeInBlock = FMath::Min(CompressionBlockSize - OffsetInBlock, RemainingSize);
+						FMemory::Memcpy(Dst, Src.Data(), SizeInBlock);
+						OffsetInBlock = 0;
+						RemainingSize -= SizeInBlock;
+						Dst += SizeInBlock;
+					}
+
+					Request->BlockBuffers.Empty();
+					Request->Callback(Request->Buffer);
+				}
+			};
+
+			Request->ReadRequests[Idx].Reset(
+				ContainerFileHandles[PartitionIndex]->ReadRequest(PartitionOffset, RawSize, AIOP_Normal, &FileReadCallback, BlockBuffer.Data()));
+		}
+
+		return ReadRequest;
+	}
+
 	const FIoDirectoryIndexReader& GetDirectoryIndexReader() const
 	{
 		return DirectoryIndexReader;
@@ -1662,21 +1877,18 @@ private:
 		return ChunkInfo;
 	}
 
-	struct FThreadBuffers
-		: public FTlsAutoCleanup
-	{
-		TArray<uint8> CompressedBuffer;
-		TArray<uint8> UncompressedBuffer;
-	};
-
 	FIoStoreToc Toc;
 	FAES::FAESKey DecryptionKey;
 	TArray<TUniquePtr<IAsyncReadFileHandle>> ContainerFileHandles;
 	FIoDirectoryIndexReader DirectoryIndexReader;
-	static thread_local FThreadBuffers* ThreadBuffers;
 };
 
-thread_local FIoStoreReaderImpl::FThreadBuffers* FIoStoreReaderImpl::ThreadBuffers = nullptr;
+FIoStoreAsyncReadRequest::FIoStoreAsyncReadRequest()
+	: Impl(new FImpl())
+{
+}
+
+FIoStoreAsyncReadRequest::~FIoStoreAsyncReadRequest() = default;
 
 FIoStoreReader::FIoStoreReader()
 	: Impl(new FIoStoreReaderImpl())
@@ -1728,6 +1940,16 @@ TIoStatusOr<FIoBuffer> FIoStoreReader::Read(const FIoChunkId& Chunk, const FIoRe
 	return Impl->Read(Chunk, Options);
 }
 
+TIoStatusOr<FIoStoreCompressedReadResult> FIoStoreReader::ReadCompressed(const FIoChunkId& Chunk, const FIoReadOptions& Options) const
+{
+	return Impl->ReadCompressed(Chunk, Options);
+}
+
+FIoStoreAsyncReadRequest FIoStoreReader::ReadAsync(const FIoChunkId& Chunk, const FIoReadOptions& Options, FIoStoreAsyncReadRequest::FCallback&& Callback) const
+{
+	return Impl->ReadAsync(Chunk, Options, MoveTemp(Callback));
+}
+
 const FIoDirectoryIndexReader& FIoStoreReader::GetDirectoryIndexReader() const
 {
 	return Impl->GetDirectoryIndexReader();
@@ -2060,4 +2282,4 @@ void FIoStoreWriterSettings::InitializePlatformSpecificSettings(const ITargetPla
 	const TCHAR* IniSection = TEXT("/Script/UnrealEd.ProjectPackagingSettings");
 	PlatformConfig->GetInt(IniSection, TEXT("PackageCompressionMinBytesSaved"), CompressionMinBytesSaved, GGameIni);
 	PlatformConfig->GetInt(IniSection, TEXT("PackageCompressionMinPercentSaved"), CompressionMinPercentSaved, GGameIni);
-}
\ No newline at end of file
+}
diff --git a/Engine/Source/Runtime/Core/Private/IO/PackageStore.cpp b/Engine/Source/Runtime/Core/Private/IO/PackageStore.cpp
new file mode 100644
--- /dev/null
+++ b/Engine/Source/Runtime/Core/Private/IO/PackageStore.cpp
@@ -0,0 +1,33 @@
+// Copyright Epic Games, Inc. All Rights Reserved.
+
+#include "IO/PackageStore.h"
+#include "Serialization/StructuredArchive.h"
+
+FArchive& operator<<(FArchive& Ar, FPackageStoreExportInfo& ExportInfo)
+{
+	Ar << ExportInfo.ExportBundlesSize;
+	Ar << ExportInfo.ExportCount;
+	Ar << ExportInfo.ExportBundleCount;
+	Ar << ExportInfo.LoadOrder;
+
+	return Ar;
+}
+
+FArchive& operator<<(FArchive& Ar, FPackageStoreEntryResource& PackageStoreEntry)
+{
+	uint32 Flags = static_cast<uint32>(PackageStoreEntry.Flags);
+
+	Ar << Flags;
+	Ar << PackageStoreEntry.PackageName;
+	Ar << PackageStoreEntry.SourcePackageName;
+	Ar << PackageStoreEntry.Region;
+	Ar << PackageStoreEntry.ExportInfo;
+	Ar << PackageStoreEntry.ImportedPackageIds;
+
+	if (Ar.IsLoading())
+	{
+		PackageStoreEntry.Flags = static_cast<EPackageStoreEntryFlags>(Flags);
+	}
+
+	return Ar;
+}
diff --git a/Engine/Source/Runtime/Core/Private/Misc/CoreGlobals.cpp b/Engine/Source/Runtime/Core/Private/Misc/CoreGlobals.cpp
--- a/Engine/Source/Runtime/Core/Private/Misc/CoreGlobals.cpp
+++ b/Engine/Source/Runtime/Core/Private/Misc/CoreGlobals.cpp
@@ -713,3 +713,22 @@ DEFINE_LOG_CATEGORY(LogVirtualization);
 #endif
 
 #undef LOCTEXT_NAMESPACE
+
+bool IsRunningCookOnTheFly()
+{
+#if WITH_COTF
+	static struct FCookOnTheFlyCommandline
+	{
+		bool bParsed;
+		FCookOnTheFlyCommandline()
+		{
+			FString Host;
+			bParsed = FParse::Value(FCommandLine::Get(), TEXT("-CookOnTheFlyHost="), Host) || FParse::Value(FCommandLine::Get(), TEXT("-FileHostIP="), Host);
+		}
+	} CookOnTheFlyCommandline;
+	
+	return CookOnTheFlyCommandline.bParsed;
+#else
+	return false;
+#endif
+}
diff --git a/Engine/Source/Runtime/Core/Private/Windows/WindowsPlatformFile.cpp b/Engine/Source/Runtime/Core/Private/Windows/WindowsPlatformFile.cpp
--- a/Engine/Source/Runtime/Core/Private/Windows/WindowsPlatformFile.cpp
+++ b/Engine/Source/Runtime/Core/Private/Windows/WindowsPlatformFile.cpp
@@ -545,11 +545,6 @@ class FFileHandleWindows : public IFileHandle
 	 */
 	uint32 Flags;
 
-	FORCEINLINE bool IsValid()
-	{
-		return FileHandle != NULL && FileHandle != INVALID_HANDLE_VALUE;
-	}
-
 	FORCEINLINE void UpdateOverlappedPos()
 	{
 		ULARGE_INTEGER LI;
@@ -565,11 +560,23 @@ class FFileHandleWindows : public IFileHandle
 		return SetFilePointer(FileHandle, LI.LowPart, &LI.HighPart, FILE_BEGIN) != INVALID_SET_FILE_POINTER;
 	}
 
-	FORCEINLINE void UpdateFileSize()
+	void UpdateFileSize()
 	{
 		LARGE_INTEGER LI;
-		GetFileSizeEx(FileHandle, &LI);
-		FileSize = LI.QuadPart;
+		BOOL Success = GetFileSizeEx(FileHandle, &LI);
+		if (Success)
+		{
+			FileSize = LI.QuadPart;
+		}
+		else
+		{
+			// This is a rare condition but it has been observed in the wild
+			// for SMB file shares. Since it's important to know the file size
+			// we render the file handle unusable if this condition occurs
+			FileSize = -1;
+
+			UE_LOG(LogTemp, Warning, TEXT("GetFileSizeEx: Failed to get file size for handle!"));
+		}
 	}
 
 public:
@@ -605,6 +612,10 @@ public:
 #endif
 		FileHandle = NULL;
 	}
+	bool IsValid()
+	{
+		return FileHandle != NULL && FileHandle != INVALID_HANDLE_VALUE && FileSize != -1;
+	}
 	virtual int64 Tell(void) override
 	{
 		check(IsValid());
@@ -1193,7 +1204,21 @@ public:
 		if (Handle != INVALID_HANDLE_VALUE)
 		{
 			TRACE_PLATFORMFILE_END_OPEN(Handle);
-			return new FFileHandleWindows(Handle, Access, WinFlags, 0);
+
+			FFileHandleWindows* FileHandle = new FFileHandleWindows(Handle, Access, WinFlags, 0);
+
+			// Some operations can fail during the handle initialization, so we
+			// double check that the handle is valid before returning it
+			if (FileHandle->IsValid())
+			{
+				return FileHandle;
+			}
+			else
+			{
+				delete FileHandle;
+
+				return nullptr;
+			}
 		}
 #endif
 		else
diff --git a/Engine/Source/Runtime/Core/Public/CoreGlobals.h b/Engine/Source/Runtime/Core/Public/CoreGlobals.h
--- a/Engine/Source/Runtime/Core/Public/CoreGlobals.h
+++ b/Engine/Source/Runtime/Core/Public/CoreGlobals.h
@@ -734,3 +734,6 @@ extern CORE_API TArray<FMinimalName, TInlineAllocator<3>>* GCoreComplexObjectPat
 /** Array to help visualize object handles in the debugger */
 struct FObjectHandlePackageDebugData;
 extern CORE_API FObjectHandlePackageDebugData* GCoreObjectHandlePackageDebug;
+
+/** @return True if running cook-on-the-fly. */
+bool CORE_API IsRunningCookOnTheFly();
diff --git a/Engine/Source/Runtime/Core/Public/IO/IoDispatcher.h b/Engine/Source/Runtime/Core/Public/IO/IoDispatcher.h
--- a/Engine/Source/Runtime/Core/Public/IO/IoDispatcher.h
+++ b/Engine/Source/Runtime/Core/Public/IO/IoDispatcher.h
@@ -11,7 +11,9 @@
 #include "Templates/UnrealTemplate.h"
 #include "Templates/TypeCompatibleBytes.h"
 #include "HAL/PlatformAtomics.h"
+#include "Memory/MemoryView.h"
 #include "Misc/SecureHash.h"
+#include "IO/IoHash.h"
 #include "Misc/AES.h"
 #include "Misc/IEngineCrypto.h"
 #include "Serialization/FileRegions.h"
@@ -622,7 +624,8 @@ public:
 	static FIoChunkHash HashBuffer(const void* Data, uint64 DataSize)
 	{
 		FIoChunkHash Result;
-		FSHA1::HashBuffer(Data, DataSize, Result.Hash);
+		FIoHash IoHash = FIoHash::HashBuffer(Data, DataSize);
+		FMemory::Memcpy(Result.Hash, &IoHash, sizeof IoHash);
 		FMemory::Memset(Result.Hash + 20, 0, 12);
 		return Result;
 	}
@@ -655,6 +658,13 @@ public:
 		return Ar;
 	}
 
+	template <typename CharType>
+	friend TStringBuilderBase<CharType>& operator<<(TStringBuilderBase<CharType>& Builder, const FIoChunkId& ChunkId)
+	{
+		UE::String::BytesToHexLower(ChunkId.Id, Builder);
+		return Builder;
+	}
+
 	inline bool operator ==(const FIoChunkId& Rhs) const
 	{
 		return 0 == FMemory::Memcmp(Id, Rhs.Id, sizeof Id);
@@ -671,11 +681,20 @@ public:
 		FMemory::Memcpy(Id, InIdPtr, sizeof Id);
 	}
 
+	void Set(FMemoryView InView)
+	{
+		check(InView.GetSize() == sizeof Id);
+		FMemory::Memcpy(Id, InView.GetData(), sizeof Id);
+	}
+
 	inline bool IsValid() const
 	{
 		return *this != InvalidChunkId;
 	}
 
+	inline const uint8* GetData() const { return Id; }
+	inline uint32		GetSize() const { return sizeof Id; }
+
 private:
 	static inline FIoChunkId CreateEmptyId()
 	{
@@ -691,33 +710,39 @@ private:
 
 /**
  * Addressable chunk types.
+ * 
+ * The enumerators have explicitly defined values here to encourage backward/forward
+ * compatible updates. 
+ * 
+ * Also note that for certain discriminators, Zen Store will assume certain things
+ * about the structure of the chunk ID so changes must be made carefully.
+ * 
  */
 enum class EIoChunkType : uint8
 {
-	Invalid,
-	InstallManifest,
-	ExportBundleData,
-	BulkData,
-	OptionalBulkData,
-	MemoryMappedBulkData,
-	LoaderGlobalMeta,
-	LoaderInitialLoadMeta,
-	LoaderGlobalNames,
-	LoaderGlobalNameHashes,
-	ContainerHeader,
-	ShaderCodeLibrary,
-	ShaderCode
+	Invalid = 0,
+	ExportBundleData = 1,
+	BulkData = 2,
+	OptionalBulkData = 3,
+	MemoryMappedBulkData = 4,
+	ScriptObjects = 5,
+	ContainerHeader = 6,
+	ExternalFile = 7,
+	ShaderCodeLibrary = 8,
+	ShaderCode = 9
 };
 
 /**
- * Creates a chunk identifier,
+ * Creates a chunk identifier (generic -- prefer specialized versions where possible)
  */
 static FIoChunkId CreateIoChunkId(uint64 ChunkId, uint16 ChunkIndex, EIoChunkType IoChunkType)
 {
+	checkSlow(IoChunkType != EIoChunkType::ExternalFile);	// Use CreateExternalFileChunkId() instead
+
 	uint8 Data[12] = {0};
 
 	*reinterpret_cast<uint64*>(&Data[0]) = ChunkId;
-	*reinterpret_cast<uint16*>(&Data[8]) = ChunkIndex;
+	*reinterpret_cast<uint16*>(&Data[8]) = NETWORK_ORDER16(ChunkIndex);
 	*reinterpret_cast<uint8*>(&Data[11]) = static_cast<uint8>(IoChunkType);
 
 	FIoChunkId IoChunkId;
@@ -726,6 +751,45 @@ static FIoChunkId CreateIoChunkId(uint64 ChunkId, uint16 ChunkIndex, EIoChunkTyp
 	return IoChunkId;
 }
 
+/**
+ * Create an "external file" chunk identifier
+ * 
+ * Layout (the index is big endian for better lexical sort behavior):
+ * 
+ *   00  01  02  03  04  05  06  07  08  09  10  11
+ * +---+---+---+---+---+---+---+---+---+---+---+---+
+ * |      ContainerId              | FileIndex | 7 |
+ * +---+---+---+---+---+---+---+---+---+---+---+---+
+ * 
+ */
+static inline FIoChunkId CreateExternalFileChunkId(uint64 ContainerId, uint32 FileIndex)
+{
+	check(FileIndex < (1 << 24));
+
+	uint8 Data[12] = { 0 };
+
+	Data[11] = uint8(EIoChunkType::ExternalFile);
+
+	memcpy(&Data[0], &ContainerId, sizeof ContainerId);
+
+	const uint32 NoFileIndex = NETWORK_ORDER32(FileIndex);
+	memcpy(&Data[8], reinterpret_cast<const uint8*>(&NoFileIndex) + 1, 3);
+
+	FIoChunkId IoChunkId;
+	IoChunkId.Set(Data, 12);
+
+	return IoChunkId;
+}
+
+inline uint32 FileIndexFromChunkId(FIoChunkId ChunkId)
+{
+	// See encoding details in CreateExternalFileChunkId()
+	uint32 FileIndex = 0;
+	memcpy(&FileIndex, ChunkId.GetData() + 7, sizeof FileIndex);
+	FileIndex = NETWORK_ORDER32(FileIndex) & 0xffFFff;
+	return FileIndex;
+}
+
 //////////////////////////////////////////////////////////////////////////
 
 class FIoReadOptions
@@ -1173,6 +1237,37 @@ struct FIoStoreTocChunkInfo
 	bool bIsCompressed;
 };
 
+struct FIoStoreCompressedBlockInfo
+{
+	FName CompressionMethod;
+	uint32 CompressedSize;
+	uint32 UncompressedSize;
+};
+
+struct FIoStoreCompressedReadResult
+{
+	FIoBuffer IoBuffer;
+	TArray<FIoStoreCompressedBlockInfo> Blocks;
+	uint64 UncompressedOffset;
+	uint64 UncompressedSize;
+};
+
+class FIoStoreAsyncReadRequest
+{
+public:
+	using FCallback = TFunction<void(TIoStatusOr<FIoBuffer>)>;
+
+	CORE_API FIoStoreAsyncReadRequest();
+	CORE_API ~FIoStoreAsyncReadRequest();
+
+private:
+	friend class FIoStoreReaderImpl;
+	class FImpl;
+	using FImplPtr = TSharedPtr<FImpl, ESPMode::ThreadSafe>;
+	
+	FImplPtr Impl;
+};
+
 class FIoStoreReader
 {
 public:
@@ -1187,6 +1282,8 @@ public:
 	CORE_API TIoStatusOr<FIoStoreTocChunkInfo> GetChunkInfo(const FIoChunkId& Chunk) const;
 	CORE_API TIoStatusOr<FIoStoreTocChunkInfo> GetChunkInfo(const uint32 TocEntryIndex) const;
 	CORE_API TIoStatusOr<FIoBuffer> Read(const FIoChunkId& Chunk, const FIoReadOptions& Options) const;
+	CORE_API TIoStatusOr<FIoStoreCompressedReadResult> ReadCompressed(const FIoChunkId& Chunk, const FIoReadOptions& Options) const;
+	CORE_API FIoStoreAsyncReadRequest ReadAsync(const FIoChunkId& Chunk, const FIoReadOptions& Options, FIoStoreAsyncReadRequest::FCallback&& Callback) const;
 
 	CORE_API const FIoDirectoryIndexReader& GetDirectoryIndexReader() const;
 
diff --git a/Engine/Source/Runtime/Core/Public/IO/PackageStore.h b/Engine/Source/Runtime/Core/Public/IO/PackageStore.h
new file mode 100644
--- /dev/null
+++ b/Engine/Source/Runtime/Core/Public/IO/PackageStore.h
@@ -0,0 +1,193 @@
+// Copyright Epic Games, Inc. All Rights Reserved.
+
+#pragma once
+
+#include "CoreMinimal.h"
+#include "PackageId.h"
+#include "Misc/SecureHash.h"
+
+class FArchive;
+class FStructuredArchiveSlot;
+
+/**
+ * Package export information.
+ */
+struct FPackageStoreExportInfo
+{
+	uint64 ExportBundlesSize = 0;
+	int32 ExportCount = 0;
+	int32 ExportBundleCount = 0;
+	uint32 LoadOrder = 0;
+
+	CORE_API friend FArchive& operator<<(FArchive& Ar, FPackageStoreExportInfo& ExportInfo);
+};
+
+/**
+ * Package store entry.
+ */ 
+struct FPackageStoreEntry
+{
+	FPackageStoreExportInfo ExportInfo;
+	TArrayView<const FPackageId> ImportedPackageIds;
+	TArrayView<const FSHAHash> ShaderMapHashes;
+};
+
+/**
+ * Package store entry flags.
+ */
+enum class EPackageStoreEntryFlags : uint32
+{
+	None		= 0,
+	Redirected	= 0x01,
+};
+ENUM_CLASS_FLAGS(EPackageStoreEntryFlags);
+
+/**
+ * Package store entry resource.
+ *
+ * This is a non-optimized serializable version
+ * of a package store entry. Used when cooking
+ * and when running cook-on-the-fly.
+ */
+struct FPackageStoreEntryResource
+{
+	/** The package store entry flags. */
+	EPackageStoreEntryFlags Flags;
+	/** The package name. */
+	FName PackageName;
+	/** Used for localized and redirected packages. */
+	FName SourcePackageName;
+	/** Region name for localized packages. */
+	FName Region;
+	/** The package export information. */
+	FPackageStoreExportInfo ExportInfo;
+	/** Imported package IDs. */
+	TArray<FPackageId> ImportedPackageIds;
+	/** Referenced shader map hashes. */
+	TArray<FSHAHash> ShaderMapHashes;
+
+	/** Returns the package ID. */
+	FPackageId GetPackageId() const
+	{
+		return FPackageId::FromName(PackageName);
+	}
+
+	/** Returns the source package ID. */
+	FPackageId GetSourcePackageId() const
+	{
+		return SourcePackageName.IsNone() ? FPackageId() : FPackageId::FromName(SourcePackageName);
+	}
+
+	/** Returns whether this package is redirected. */
+	bool IsRedirected() const
+	{
+		return EnumHasAnyFlags(Flags, EPackageStoreEntryFlags::Redirected); 
+	}
+
+	CORE_API friend FArchive& operator<<(FArchive& Ar, FPackageStoreEntryResource& PackageStoreEntry);
+};
+
+/**
+ * Package store entry status.
+ */
+enum class EPackageStoreEntryStatus
+{
+	None,
+	Ok,
+	Pending,
+	Missing,
+};
+
+/**
+ * Represents a handle to package information.
+ */
+class FPackageStoreEntryHandle
+{
+	static constexpr uint64 HandleBits = 62ull;
+	static constexpr uint64 HandleMask = (1ull << HandleBits) - 1ull;
+	static constexpr uint64 StatusMask = ~HandleMask;
+	static constexpr uint64 StatusShift = HandleBits;
+	static constexpr uint64 InvalidHandle = 0;
+
+	uint64 Handle = InvalidHandle;
+
+	inline explicit FPackageStoreEntryHandle(const uint64 InHandle)
+		: Handle(InHandle) { }
+
+public:
+	FPackageStoreEntryHandle() = default;
+
+	inline uint64 Value() const
+	{
+		return (Handle & HandleMask);
+	}
+
+	inline bool IsValid() const
+	{
+		return Value() != InvalidHandle;
+	}
+
+	inline EPackageStoreEntryStatus Status() const
+	{
+		return static_cast<EPackageStoreEntryStatus>((Handle & StatusMask) >> StatusShift);
+	}
+
+	inline bool operator<(FPackageStoreEntryHandle Other) const
+	{
+		return Handle < Other.Handle;
+	}
+
+	inline bool operator==(FPackageStoreEntryHandle Other) const
+	{
+		return Handle == Other.Handle;
+	}
+	
+	inline bool operator!=(FPackageStoreEntryHandle Other) const
+	{
+		return Handle != Other.Handle;
+	}
+
+	inline friend uint32 GetTypeHash(const FPackageStoreEntryHandle& Entry)
+	{
+		return uint32(Entry.Handle);
+	}
+
+	operator bool() const
+	{
+		return IsValid();
+	}
+
+	static FPackageStoreEntryHandle Create(uint64 InHandle, EPackageStoreEntryStatus EntryStatus = EPackageStoreEntryStatus::Ok)
+	{
+		check(!(InHandle & StatusMask));
+		return FPackageStoreEntryHandle((static_cast<uint64>(EntryStatus) << StatusShift) | (InHandle & HandleMask));
+	}
+};
+
+/**
+ * Stores information about available packages that can be loaded.
+ */
+class IPackageStore
+{
+public:
+	/* Destructor. */
+	virtual ~IPackageStore() { }
+
+	/* Initialize the package store. */
+	virtual void Initialize() = 0;
+
+	/* Returns whether the package exists. */
+	virtual bool DoesPackageExist(FPackageId PackageId) = 0;
+
+	/* Returns the package store entry for the specified package ID. */
+	virtual FPackageStoreEntryHandle GetPackageEntryHandle(FPackageId PackageId, const FName& PackageName = NAME_None) = 0;
+
+	/* Returns the package store entry data with export info and imported packages for the specified package ID. */
+	virtual FPackageStoreEntry GetPackageEntry(FPackageStoreEntryHandle Handle) = 0;
+
+	/* Returns the redirected package ID for the specified package ID. */
+	virtual FPackageId GetRedirectedPackageId(FPackageId PackageId) = 0;
+
+	/* Returns whether the package ID is a redirect. */
+	virtual bool IsRedirect(FPackageId PackageId) = 0;
+};
diff --git a/Engine/Source/Runtime/Core/Public/IO/PackageStoreWriter.h b/Engine/Source/Runtime/Core/Public/IO/PackageStoreWriter.h
new file mode 100644
--- /dev/null
+++ b/Engine/Source/Runtime/Core/Public/IO/PackageStoreWriter.h
@@ -0,0 +1,136 @@
+// Copyright Epic Games, Inc. All Rights Reserved.
+
+#pragma once
+
+#include "IO/IoDispatcher.h"
+
+struct FPackageStoreEntryResource;
+
+class IPackageStoreWriter
+{
+public:
+	virtual ~IPackageStoreWriter() = default;
+
+	struct FPackageBaseInfo
+	{
+		FName	PackageName;
+	};
+
+	/** Mark the beginning of a package store transaction for the specified package
+
+		This must be called before any data is produced for a given package
+	  */
+	virtual void BeginPackage(const FPackageBaseInfo& Info) = 0;
+
+	/** Finalize a package started with BeginPackage()
+	  */
+	virtual void CommitPackage(const FPackageBaseInfo& Info) = 0;
+
+	struct FPackageInfo
+	{
+		FName		PackageName;
+		FString		LooseFilePath;
+		uint64		HeaderSize;
+		FIoChunkId	ChunkId = FIoChunkId::InvalidChunkId;
+	};
+
+	/** Write package data (exports and serialized header)
+
+		This may only be called after a BeginPackage() call has been made
+		to signal the start of a package store transaction
+	  */
+	virtual void WritePackageData(const FPackageInfo& Info, const FIoBuffer& PackageData, const TArray<FFileRegion>& FileRegions) = 0;
+
+	struct FBulkDataInfo
+	{
+		enum EType
+		{
+			Standard,
+			Mmap,
+			Optional
+		};
+
+		FName		PackageName;
+		EType		BulkdataType = Standard;
+		FString		LooseFilePath;
+		FIoChunkId	ChunkId = FIoChunkId::InvalidChunkId;
+	};
+
+	/** Write bulk data for the current package
+	  */
+	virtual void WriteBulkdata(const FBulkDataInfo& Info, const FIoBuffer& BulkData, const TArray<FFileRegion>& FileRegions) = 0;
+
+	struct FAdditionalFileInfo
+	{
+		FName	PackageName;
+		FString	Filename;
+		FIoChunkId ChunkId;
+	};
+
+	virtual bool WriteAdditionalFile(const FAdditionalFileInfo& Info, const FIoBuffer& FileData) = 0;
+
+	struct FCookInfo
+	{
+		enum ECookMode
+		{
+			CookByTheBookMode,
+			CookOnTheFlyMode
+		};
+
+		ECookMode CookMode = CookByTheBookMode;
+	};
+
+	/** Signal the start of a cooking pass
+
+		Package data may only be produced after BeginCook() has been called and
+		before EndCook() is called
+	  */
+	virtual void BeginCook(const FCookInfo& Info) = 0;
+
+	/** Signal the end of a cooking pass.
+	  */
+	virtual void EndCook() = 0;
+
+	/**
+	 * Returns all cooked package store entries.
+	 */
+	virtual void GetEntries(TFunction<void(TArrayView<const FPackageStoreEntryResource>)>&&) = 0;
+
+	/**
+	 * Package commit event arguments
+	 */
+	struct FCommitEventArgs
+	{
+		FName PlatformName;
+		FName PackageName;
+		int32 EntryIndex = INDEX_NONE;
+		TArrayView<const FPackageStoreEntryResource> Entries;
+		TArray<FAdditionalFileInfo> AdditionalFiles;
+	};
+
+	/**
+	 * Broadcasted after a package has been committed, i.e cooked.
+	 */
+	DECLARE_EVENT_OneParam(IPackageStoreWriter, FCommitEvent, const FCommitEventArgs&);
+	virtual FCommitEvent& OnCommit() = 0;
+
+	/**
+	 * Flush any outstanding writes.
+	 */
+	virtual void Flush() = 0;
+};
+
+static inline const ANSICHAR* LexToString(IPackageStoreWriter::FBulkDataInfo::EType Value)
+{
+	switch (Value)
+	{
+	case IPackageStoreWriter::FBulkDataInfo::Standard:
+		return "Standard";
+	case IPackageStoreWriter::FBulkDataInfo::Mmap:
+		return "Mmap";
+	case IPackageStoreWriter::FBulkDataInfo::Optional:
+		return "Optional";
+	}
+
+	return "Unknown";
+}
diff --git a/Engine/Source/Runtime/Core/Public/Misc/Build.h b/Engine/Source/Runtime/Core/Public/Misc/Build.h
--- a/Engine/Source/Runtime/Core/Public/Misc/Build.h
+++ b/Engine/Source/Runtime/Core/Public/Misc/Build.h
@@ -440,5 +440,9 @@
 	#define PRELOAD_BINARY_CONFIG 1
 #endif
 
+#ifndef WITH_COTF
+	#define WITH_COTF ((WITH_ENGINE) && !(IS_PROGRAM || WITH_EDITOR || UE_SERVER || UE_BUILD_SHIPPING))
+#endif
+
 // Controls if the config system can stores configs for other platforms than the running one
 #define ALLOW_OTHER_PLATFORM_CONFIG		WITH_UNREAL_DEVELOPER_TOOLS
diff --git a/Engine/Source/Runtime/Core/Public/Serialization/LargeMemoryWriter.h b/Engine/Source/Runtime/Core/Public/Serialization/LargeMemoryWriter.h
--- a/Engine/Source/Runtime/Core/Public/Serialization/LargeMemoryWriter.h
+++ b/Engine/Source/Runtime/Core/Public/Serialization/LargeMemoryWriter.h
@@ -4,6 +4,7 @@
 
 #include "CoreTypes.h"
 #include "Containers/UnrealString.h"
+#include "Memory/MemoryView.h"
 #include "UObject/NameTypes.h"
 #include "Serialization/MemoryArchive.h"
 #include "Serialization/LargeMemoryData.h"
@@ -40,6 +41,15 @@ public:
 	 */
 	uint8* GetData() const;
 
+	/**
+	 * Returns a view on the written data
+	 * 
+	 * The view does not own the memory, so you must make sure you keep the memory writer
+	 * alive while you are using the returned view.
+	 * 
+	 */
+	inline FMemoryView GetView() const { return MakeMemoryView(GetData(), Data.GetSize()); }
+
 	/** 
 	 * Releases ownership of the written data
 	 *
