COTF2: Don't send heartbeat message while processing requests

#rnx

diff --git a/Engine/Source/Runtime/CookOnTheFly/Private/CookOnTheFlyConnectionServer.cpp b/Engine/Source/Runtime/CookOnTheFly/Private/CookOnTheFlyConnectionServer.cpp
--- a/Engine/Source/Runtime/CookOnTheFly/Private/CookOnTheFlyConnectionServer.cpp
+++ b/Engine/Source/Runtime/CookOnTheFly/Private/CookOnTheFlyConnectionServer.cpp
@@ -18,7 +18,7 @@ class FCookOnTheFlyConnectionServer final
 	: public UE::Cook::ICookOnTheFlyConnectionServer
 {
 	static constexpr uint32 ServerSenderId = ~uint32(0);
-	static constexpr double HeartbeatTimeoutInSeconds = 120.0;
+	static constexpr double HeartbeatTimeoutInSeconds = 60 * 5;
 
 public:
 	FCookOnTheFlyConnectionServer(UE::Cook::FCookOnTheFlyServerOptions InOptions)
@@ -85,6 +85,7 @@ private:
 		TFuture<void> Thread;
 		TAtomic<bool> bIsRunning { false };
 		TAtomic<bool> bStopRequested { false };
+		TAtomic<bool> bIsProcessingRequest { false };
 		TAtomic<double> LastActivityTime {0};
 		uint32 ClientId = 0;
 		FName PlatformName;
@@ -155,9 +156,12 @@ private:
 			{
 				UE_LOG(LogCotfConnectionServer, Warning, TEXT("Failed to send message '%s' to client '%s' (Id='%d', Platform='%s')"),
 					LexToString(Message.GetHeader().MessageType), *Client->PeerAddr->ToString(true), Client->ClientId, *Client->PlatformName.ToString());
-
+				
+				Client->bIsRunning = false;
 				bBroadcasted = false;
 			}
+
+			Client->LastActivityTime = FPlatformTime::Seconds();
 		}
 
 		return bBroadcasted;
@@ -214,7 +218,7 @@ private:
 
 					const double SecondsSinceLastActivity = FPlatformTime::Seconds() - Client->LastActivityTime.Load();
 
-					if (SecondsSinceLastActivity > HeartbeatTimeoutInSeconds)
+					if (SecondsSinceLastActivity > HeartbeatTimeoutInSeconds && !Client->bIsProcessingRequest)
 					{
 						Client->LastActivityTime = FPlatformTime::Seconds();
 
@@ -276,7 +280,9 @@ private:
 	bool ProcesseRequest(FClient& Client)
 	{
 		using namespace UE::Cook;
-		
+
+		Client.bIsProcessingRequest = false;
+
 		FArrayReader RequestPayload;
 		if (!FNFSMessageHeader::ReceivePayload(RequestPayload, FSimpleAbstractSocket_FSocket(Client.Socket)))
 		{
@@ -284,6 +290,8 @@ private:
 			return false;
 		}
 
+		Client.bIsProcessingRequest = true;
+
 		FCookOnTheFlyRequest Request;
 		RequestPayload << Request;
 
