Temp fix for COTF crashing in cooked build. Submitting on behalf of Per.Larsson
#rnx

diff --git a/Engine/Source/Runtime/CookOnTheFly/Private/CookOnTheFly.cpp b/Engine/Source/Runtime/CookOnTheFly/Private/CookOnTheFly.cpp
--- a/Engine/Source/Runtime/CookOnTheFly/Private/CookOnTheFly.cpp
+++ b/Engine/Source/Runtime/CookOnTheFly/Private/CookOnTheFly.cpp
@@ -87,6 +87,11 @@ bool GetCookOnTheFlyHost(UE::Cook::FCookOnTheFlyHostOptions& OutHostOptions)
 
 bool SendCookOnTheFlyRequest(ECookOnTheFlyMessage RequestType, TFunction<void(FArchive&)>&& FillRequest, TFunction<bool(FArchive&)>&& ProcessResponse)
 {
+	if (!IsRunningCookOnTheFly())
+	{
+		return false;
+	}
+
 	static FString FileHostIP;
 	static bool bUseFileServer = FParse::Value(FCommandLine::Get(), TEXT("filehostip"), FileHostIP);
 
