PackageStore refactor
* Decouple container concept from IoDispatcher
* Decoiuple PackageStore implementation from AsyncLoading2
* Restore ucas unmount fix that got kist when merrging from UE4
* Fix packages being left in the PackageStiore even after unmounting contaiiners

#rnx

diff --git a/Engine/Source/Developer/IoStoreUtilities/Internal/PackageStoreOptimizer.h b/Engine/Source/Developer/IoStoreUtilities/Internal/PackageStoreOptimizer.h
--- a/Engine/Source/Developer/IoStoreUtilities/Internal/PackageStoreOptimizer.h
+++ b/Engine/Source/Developer/IoStoreUtilities/Internal/PackageStoreOptimizer.h
@@ -7,6 +7,7 @@
 #include "IO/IoDispatcher.h"
 #include "UObject/PackageFileSummary.h"
 #include "UObject/UObjectMarks.h"
+#include "IO/IoContainerHeader.h"
 
 class FBufferWriter;
 
@@ -329,7 +330,7 @@ public:
 	void FinalizePackage(FPackageStorePackage* Package);
 	FIoBuffer CreatePackageBuffer(const FPackageStorePackage* Package, const FIoBuffer& CookedExportsBuffer, TArray<FFileRegion>* InOutFileRegions) const;
 	FPackageStoreEntryResource CreatePackageStoreEntry(const FPackageStorePackage* Package) const;
-	FContainerHeader CreateContainerHeader(const FIoContainerId& ContainerId, TArrayView<const FPackageStoreEntryResource> PackageStoreEntries) const;
+	FIoContainerHeader CreateContainerHeader(const FIoContainerId& ContainerId, TArrayView<const FPackageStoreEntryResource> PackageStoreEntries) const;
 	IOSTOREUTILITIES_API FIoBuffer CreateScriptObjectsBuffer() const;
 	void LoadScriptObjectsBuffer(const FIoBuffer& ScriptObjectsBuffer);
 	void ProcessRedirects(const TMap<FPackageId, FPackageStorePackage*>& PackagesMap, bool bIsBuildingDLC) const;
diff --git a/Engine/Source/Developer/IoStoreUtilities/Private/IoStoreUtilities.cpp b/Engine/Source/Developer/IoStoreUtilities/Private/IoStoreUtilities.cpp
--- a/Engine/Source/Developer/IoStoreUtilities/Private/IoStoreUtilities.cpp
+++ b/Engine/Source/Developer/IoStoreUtilities/Private/IoStoreUtilities.cpp
@@ -28,7 +28,6 @@
 #include "Serialization/AsyncLoading2.h"
 #include "Serialization/ArrayReader.h"
 #include "IO/PackageStore.h"
-#include "Serialization/FilePackageStore.h"
 #include "UObject/Class.h"
 #include "UObject/NameBatchSerialization.h"
 #include "UObject/PackageFileSummary.h"
@@ -56,6 +55,7 @@
 #include "ZenStoreHttpClient.h"
 #include "IPlatformFilePak.h"
 #include "ZenStoreWriter.h"
+#include "IO/IoContainerHeader.h"
 #include "ProfilingDebugging/CountersTrace.h"
 
 //PRAGMA_DISABLE_OPTIMIZATION
@@ -387,7 +387,7 @@ public:
 			}
 
 			FMemoryReaderView Ar(MakeArrayView(HeaderBuffer.ValueOrDie().Data(), HeaderBuffer.ValueOrDie().DataSize()));
-			FContainerHeader ContainerHeader;
+			FIoContainerHeader ContainerHeader;
 			Ar << ContainerHeader;
 
 			PackageEntries.Reserve(ContainerHeader.PackageCount);
@@ -631,7 +631,7 @@ struct FIoStoreArguments
 struct FContainerTargetSpec
 {
 	FIoContainerId ContainerId;
-	FContainerHeader Header;
+	FIoContainerHeader Header;
 	FName Name;
 	FGuid EncryptionKeyGuid;
 	FString OutputPath;
@@ -2597,6 +2597,7 @@ int32 CreateTarget(const FIoStoreArguments& Arguments, const FIoStoreWriterSetti
 			
 			FIoWriteOptions WriteOptions;
 			WriteOptions.DebugName = TEXT("ContainerHeader");
+			WriteOptions.bForceUncompressed = true;
 			ContainerTarget->IoStoreWriter->Append(
 				CreateIoChunkId(ContainerTarget->ContainerId.Value(), 0, EIoChunkType::ContainerHeader),
 				ContainerHeaderBuffer,
@@ -3177,7 +3178,7 @@ int32 Describe(
 		TArray<FPackageDesc*> Packages;
 		FIoStoreReader* Reader = nullptr;
 		FCulturePackageMap RawCulturePackageMap;
-		TArray<FContainerHeaderPackageRedirect> RawPackageRedirects;
+		TArray<FIoContainerHeaderPackageRedirect> RawPackageRedirects;
 	};
 
 	TArray<FLoadContainerHeaderJob> LoadContainerHeaderJobs;
@@ -3220,7 +3221,7 @@ int32 Describe(
 		if (IoBuffer.IsOk())
 		{
 			FMemoryReaderView Ar(MakeArrayView(IoBuffer.ValueOrDie().Data(), IoBuffer.ValueOrDie().DataSize()));
-			FContainerHeader ContainerHeader;
+			FIoContainerHeader ContainerHeader;
 			Ar << ContainerHeader;
 
 			Job.RawCulturePackageMap = ContainerHeader.CulturePackageMap;
@@ -4246,7 +4247,7 @@ int32 Staged2Zen(const FString& BuildPath, const FKeyChain& KeyChain, const FStr
 	{
 		FIoBuffer ContainerHeaderBuffer = ContainerHeaderChunk.Key->Read(ContainerHeaderChunk.Value, FIoReadOptions()).ValueOrDie();
 		FMemoryReaderView Ar(MakeArrayView(ContainerHeaderBuffer.Data(), ContainerHeaderBuffer.DataSize()));
-		FContainerHeader ContainerHeader;
+		FIoContainerHeader ContainerHeader;
 		Ar << ContainerHeader;
 		
 		const FFilePackageStoreEntry* StoreEntry = reinterpret_cast<const FFilePackageStoreEntry*>(ContainerHeader.StoreEntries.GetData());
diff --git a/Engine/Source/Developer/IoStoreUtilities/Private/PackageStoreOptimizer.cpp b/Engine/Source/Developer/IoStoreUtilities/Private/PackageStoreOptimizer.cpp
--- a/Engine/Source/Developer/IoStoreUtilities/Private/PackageStoreOptimizer.cpp
+++ b/Engine/Source/Developer/IoStoreUtilities/Private/PackageStoreOptimizer.cpp
@@ -8,7 +8,6 @@
 #include "UObject/NameBatchSerialization.h"
 #include "Containers/Map.h"
 #include "UObject/UObjectHash.h"
-#include "Serialization/FilePackageStore.h"
 #include "Interfaces/ITargetPlatform.h"
 #include "UObject/Object.h"
 #include "Serialization/MemoryReader.h"
@@ -1726,9 +1725,9 @@ FPackageStoreEntryResource FPackageStoreOptimizer::CreatePackageStoreEntry(const
 	return Result;
 }
 
-FContainerHeader FPackageStoreOptimizer::CreateContainerHeader(const FIoContainerId& ContainerId, TArrayView<const FPackageStoreEntryResource> PackageStoreEntries) const
+FIoContainerHeader FPackageStoreOptimizer::CreateContainerHeader(const FIoContainerId& ContainerId, TArrayView<const FPackageStoreEntryResource> PackageStoreEntries) const
 {
-	FContainerHeader Header;
+	FIoContainerHeader Header;
 	Header.ContainerId = ContainerId;
 	Header.PackageCount = PackageStoreEntries.Num();
 	int32 StoreTocSize = Header.PackageCount * sizeof(FFilePackageStoreEntry);
diff --git a/Engine/Source/Developer/IoStoreUtilities/Private/ZenStoreWriter.cpp b/Engine/Source/Developer/IoStoreUtilities/Private/ZenStoreWriter.cpp
--- a/Engine/Source/Developer/IoStoreUtilities/Private/ZenStoreWriter.cpp
+++ b/Engine/Source/Developer/IoStoreUtilities/Private/ZenStoreWriter.cpp
@@ -1054,7 +1054,7 @@ void FZenStoreWriter::CreateProjectMetaData(FCbPackage& Pkg, FCbWriter& PackageO
 			FIoBuffer HeaderBuffer;
 
 			{
-				FContainerHeader Header = PackageStoreOptimizer->CreateContainerHeader(ContainerId, PackageStoreEntries);
+				FIoContainerHeader Header = PackageStoreOptimizer->CreateContainerHeader(ContainerId, PackageStoreEntries);
 				FLargeMemoryWriter HeaderAr(0, true);
 				HeaderAr << Header;
 				int64 DataSize = HeaderAr.TotalSize();
diff --git a/Engine/Source/Editor/UnrealEd/Private/Cooker/IoStoreCookOnTheFlyRequestManager.cpp b/Engine/Source/Editor/UnrealEd/Private/Cooker/IoStoreCookOnTheFlyRequestManager.cpp
--- a/Engine/Source/Editor/UnrealEd/Private/Cooker/IoStoreCookOnTheFlyRequestManager.cpp
+++ b/Engine/Source/Editor/UnrealEd/Private/Cooker/IoStoreCookOnTheFlyRequestManager.cpp
@@ -4,7 +4,7 @@
 #include "CookOnTheFlyServerInterface.h"
 #include "CookOnTheFly.h"
 #include "Modules/ModuleManager.h"
-#include "Serialization/CookOnTheFlyPackageStore.h"
+#include "CookOnTheFlyPackageStore.h"
 #include "PackageStoreWriter.h"
 #include "ShaderCompiler.h"
 #include "HAL/PlatformTime.h"
diff --git a/Engine/Source/Runtime/CookOnTheFly/Private/CookOnTheFlyPackageStore.cpp b/Engine/Source/Runtime/CookOnTheFly/Private/CookOnTheFlyPackageStore.cpp
new file mode 100644
--- /dev/null
+++ b/Engine/Source/Runtime/CookOnTheFly/Private/CookOnTheFlyPackageStore.cpp
@@ -0,0 +1,319 @@
+// Copyright Epic Games, Inc. All Rights Reserved.
+
+#include "CookOnTheFlyPackageStore.h"
+
+#if WITH_COTF
+
+#include "CookOnTheFly.h"
+#include "HAL/PlatformTime.h"
+#include "Containers/ChunkedArray.h"
+#include "Serialization/MemoryReader.h"
+
+namespace UE { namespace PackageStore { namespace Messaging
+{
+
+FArchive& operator<<(FArchive& Ar, FPackageStoreData& PackageStoreData)
+{
+	Ar << PackageStoreData.CookedPackages;
+	Ar << PackageStoreData.FailedPackages;
+	Ar << PackageStoreData.TotalCookedPackages;
+	Ar << PackageStoreData.TotalFailedPackages;
+
+	return Ar;
+}
+
+FArchive& operator<<(FArchive& Ar, FCookPackageRequest& Request)
+{
+	Ar << Request.PackageName;
+
+	return Ar;
+}
+
+FArchive& operator<<(FArchive& Ar, FCookPackageResponse& Response)
+{
+	uint32 Status = static_cast<uint32>(Response.Status);
+	Ar << Status;
+
+	if (Ar.IsLoading())
+	{
+		Response.Status = static_cast<EPackageStoreEntryStatus>(Status);
+	}
+
+	return Ar;
+}
+
+}}} // namesapce UE::PackageStore::Messaging
+
+class FCookOnTheFlyPackageStore final
+	: public FPackageStoreBase
+{
+public:
+	FCookOnTheFlyPackageStore(UE::Cook::ICookOnTheFlyServerConnection& InCookOnTheFlyServerConnection)
+		: CookOnTheFlyServerConnection(InCookOnTheFlyServerConnection)
+	{
+		// Index zero is invalid
+		PackageEntries.Add();
+		CookOnTheFlyServerConnection.OnMessage().AddRaw(this, &FCookOnTheFlyPackageStore::OnCookOnTheFlyMessage);
+
+		using namespace UE::Cook;
+		using namespace UE::PackageStore::Messaging;
+
+		FCookOnTheFlyRequest Request(ECookOnTheFlyMessage::GetCookedPackages);
+		FCookOnTheFlyResponse Response = CookOnTheFlyServerConnection.SendRequest(Request).Get();
+
+		if (Response.IsOk())
+		{
+			FGetCookedPackagesResponse GetCookedPackagesResponse = Response.GetBodyAs<FGetCookedPackagesResponse>();
+			FPackageStoreData& PackageStoreData = GetCookedPackagesResponse.PackageStoreData;
+
+			UE_LOG(LogCookOnTheFly, Log, TEXT("Got '%d' cooked and '%d' failed packages from server"),
+				PackageStoreData.CookedPackages.Num(), PackageStoreData.FailedPackages.Num());
+
+			AddPackages(MoveTemp(PackageStoreData.CookedPackages), MoveTemp(PackageStoreData.FailedPackages));
+
+			LastServerActivtyTime = FPlatformTime::Seconds();
+		}
+		else
+		{
+			UE_LOG(LogCookOnTheFly, Warning, TEXT("Failed to send '%s' request"), LexToString(Request.GetHeader().MessageType));
+		}
+	}
+
+	virtual void Initialize() override
+	{
+
+	}
+
+	virtual void Lock() override
+	{
+	}
+
+	virtual void Unlock() override
+	{
+	}
+
+	virtual bool DoesPackageExist(FPackageId PackageId) override
+	{
+		FScopeLock _(&CriticalSection);
+		FEntryInfo EntryInfo = PackageIdToEntryInfo.FindRef(PackageId);
+		return EntryInfo.Status != EPackageStoreEntryStatus::Missing;
+	}
+
+	virtual FPackageStoreEntryHandle GetPackageEntryHandle(FPackageId PackageId, const FName& PackageName) override
+	{
+		using namespace UE::Cook;
+		using namespace UE::PackageStore::Messaging;
+
+		{
+			FScopeLock _(&CriticalSection);
+			FEntryInfo& EntryInfo = PackageIdToEntryInfo.FindOrAdd(PackageId, FEntryInfo());
+
+			if (EntryInfo.Status == EPackageStoreEntryStatus::Ok)
+			{
+				check(EntryInfo.EntryIndex != INDEX_NONE);
+				return FPackageStoreEntryHandle::Create(EntryInfo.EntryIndex, EPackageStoreEntryStatus::Ok); 
+			}
+			else if (PackageName == NAME_None)
+			{
+				// When importing package(s) recursive the package name is NONE. Imported
+				// package(s) are expected to be cooked and ready. This is handled by the COTF server.
+				return FPackageStoreEntryHandle::Create(0, EPackageStoreEntryStatus::Missing); 
+			}
+			
+			if (EntryInfo.Status == EPackageStoreEntryStatus::Missing)
+			{
+				return FPackageStoreEntryHandle::Create(0, EntryInfo.Status); 
+			}
+
+			if (EntryInfo.Status == EPackageStoreEntryStatus::Pending)
+			{
+				CheckActivity();
+				return FPackageStoreEntryHandle::Create(0, EntryInfo.Status); 
+			}
+
+			// The package hasn't been requested, set the status to pending.
+			EntryInfo.PackageName = PackageName;
+			EntryInfo.Status = EPackageStoreEntryStatus::Pending;
+		}
+
+		LastClientActivtyTime = FPlatformTime::Seconds();
+		UE_LOG(LogCookOnTheFly, Verbose, TEXT("Requesting package '%s'"), *PackageName.ToString());
+
+		FCookOnTheFlyRequest Request(ECookOnTheFlyMessage::CookPackage);
+		Request.SetBodyTo(FCookPackageRequest { PackageName });
+		FCookOnTheFlyResponse Response = CookOnTheFlyServerConnection.SendRequest(Request).Get();
+
+		if (!Response.IsOk())
+		{
+			UE_LOG(LogCookOnTheFly, Warning, TEXT("Failed to send '%s' request"), LexToString(Request.GetHeader().MessageType));
+			return FPackageStoreEntryHandle::Create(0, EPackageStoreEntryStatus::Missing); 
+		}
+
+		FCookPackageResponse CookPackageResponse = Response.GetBodyAs<FCookPackageResponse>();
+
+		{
+			FScopeLock _(&CriticalSection);
+
+			if (CookPackageResponse.Status == EPackageStoreEntryStatus::Missing)
+			{
+				FEntryInfo& EntryInfo = PackageIdToEntryInfo.FindChecked(PackageId);
+				EntryInfo.Status = EPackageStoreEntryStatus::Missing;
+				return FPackageStoreEntryHandle::Create(0, EPackageStoreEntryStatus::Missing); 
+			}
+			else
+			{
+				return FPackageStoreEntryHandle::Create(0, EPackageStoreEntryStatus::Pending);
+			}
+		}
+	}
+
+	virtual FPackageStoreEntry GetPackageEntry(FPackageStoreEntryHandle Handle) override
+	{
+		check(Handle.IsValid());
+		check(Handle.Status() == EPackageStoreEntryStatus::Ok);
+
+		FScopeLock _(&CriticalSection);
+		const int32 EntryIndex = static_cast<int32>(Handle.Value());
+		const FPackageStoreEntryResource& Entry = PackageEntries[EntryIndex];
+
+		return FPackageStoreEntry
+		{
+			Entry.ExportInfo,
+			Entry.ImportedPackageIds
+		};
+	}
+
+	virtual bool GetPackageRedirectInfo(FPackageId PackageId, FName& OutSourcePackageName, FPackageId& OutRedirectedToPackageId)  override
+	{
+		return false;
+	}
+
+	void AddPackages(TArray<FPackageStoreEntryResource> Entries, TArray<FPackageId> FailedPackageIds)
+	{
+		FScopeLock _(&CriticalSection);
+		
+		for (FPackageId FailedPackageId : FailedPackageIds)
+		{
+			UE_LOG(LogCookOnTheFly, Verbose, TEXT("'0x%llX' [Failed]"), FailedPackageId.Value());
+			FEntryInfo& EntryInfo = PackageIdToEntryInfo.FindOrAdd(FailedPackageId, FEntryInfo());
+			EntryInfo.Status = EPackageStoreEntryStatus::Missing;
+			PackageStats.Failed++;
+		}
+
+		const int32 NumPackageEntries = PackageEntries.Num();
+
+		for (FPackageStoreEntryResource& Entry : Entries)
+		{
+			const FPackageId PackageId = Entry.GetPackageId();
+			FEntryInfo& EntryInfo = PackageIdToEntryInfo.FindOrAdd(PackageId, FEntryInfo());
+
+			if (EntryInfo.EntryIndex == INDEX_NONE)
+			{
+				EntryInfo.Status = EPackageStoreEntryStatus::Ok;
+				EntryInfo.EntryIndex = PackageEntries.Add();
+				PackageEntries[EntryInfo.EntryIndex] = MoveTemp(Entry);
+				PackageStats.Cooked++;
+
+				UE_LOG(LogCookOnTheFly, Verbose, TEXT("'%s' [OK] (Cooked/Failed='%d/%d')"),
+					*PackageEntries[EntryInfo.EntryIndex].PackageName.ToString(), PackageStats.Cooked.Load(), PackageStats.Failed.Load());
+			}
+		}
+	}
+
+	void OnCookOnTheFlyMessage(const UE::Cook::FCookOnTheFlyMessage& Message)
+	{
+		using namespace UE::PackageStore::Messaging;
+
+		switch (Message.GetHeader().MessageType)
+		{
+			case UE::Cook::ECookOnTheFlyMessage::PackagesCooked:
+			{
+				FPackagesCookedMessage PackagesCookedMessage = Message.GetBodyAs<FPackagesCookedMessage>();
+				FPackageStoreData& PackageStoreData = PackagesCookedMessage.PackageStoreData;
+
+				UE_LOG(LogCookOnTheFly, Verbose, TEXT("Received '%s' message, Cooked='%d', Failed='%d', Server total='%d/%d' (Cooked/Failed)"),
+					LexToString(Message.GetHeader().MessageType),
+					PackageStoreData.CookedPackages.Num(),
+					PackageStoreData.FailedPackages.Num(),
+					PackageStoreData.TotalCookedPackages,
+					PackageStoreData.TotalFailedPackages);
+
+				AddPackages(MoveTemp(PackageStoreData.CookedPackages), MoveTemp(PackageStoreData.FailedPackages));
+
+				UE_CLOG(PackageStoreData.TotalCookedPackages != PackageStats.Cooked.Load() || PackageStoreData.TotalFailedPackages != PackageStats.Failed.Load(),
+					LogCookOnTheFly, Warning, TEXT("Client/Server package mismatch, Cooked='%d/%d', Failed='%d/%d' (Client/Server)"),
+					PackageStats.Cooked.Load(), PackageStoreData.TotalCookedPackages, PackageStats.Failed.Load(), PackageStoreData.TotalFailedPackages);
+
+				PendingEntriesAdded.Broadcast();
+
+				break;
+			}
+
+			default:
+				break;
+		}
+
+		LastServerActivtyTime = FPlatformTime::Seconds();
+	}
+
+	void CheckActivity()
+	{
+		const double TimeSinceLastClientActivity = FPlatformTime::Seconds() - LastClientActivtyTime;
+		const double TimeSinceLastServerActivity = FPlatformTime::Seconds() - LastServerActivtyTime;
+		const double TimeSinceLastWarning = FPlatformTime::Seconds() - LastWarningTime;
+
+		if (TimeSinceLastClientActivity > MaxInactivityTime &&
+			TimeSinceLastServerActivity > MaxInactivityTime &&
+			TimeSinceLastWarning > TimeBetweenWarning)
+		{
+			LastWarningTime = FPlatformTime::Seconds();
+
+			UE_LOG(LogCookOnTheFly, Log, TEXT("No server response in '%.2lf' seconds"), TimeSinceLastServerActivity);
+
+			UE_LOG(LogCookOnTheFly, Log, TEXT("=== Pending Packages ==="));
+			{
+				FScopeLock _(&CriticalSection);
+				for (const auto& KeyValue : PackageIdToEntryInfo)
+				{
+					const FEntryInfo& EntryInfo = KeyValue.Value;
+					if (EntryInfo.Status == EPackageStoreEntryStatus::Pending)
+					{
+						UE_LOG(LogCookOnTheFly, Log, TEXT("%s"), *EntryInfo.PackageName.ToString());
+					}
+				}
+			}
+		}
+	}
+
+	struct FEntryInfo
+	{
+		FName PackageName;
+		EPackageStoreEntryStatus Status = EPackageStoreEntryStatus::None;
+		int32 EntryIndex = INDEX_NONE;
+	};
+
+	struct FPackageStats
+	{
+		TAtomic<uint32> Cooked {0};
+		TAtomic<uint32> Failed {0};
+	};
+
+	UE::Cook::ICookOnTheFlyServerConnection& CookOnTheFlyServerConnection;
+	FCriticalSection CriticalSection;
+	TMap<FPackageId, FEntryInfo> PackageIdToEntryInfo;
+	TChunkedArray<FPackageStoreEntryResource> PackageEntries;
+	FPackageStats PackageStats;
+
+	const double MaxInactivityTime = 20;
+	const double TimeBetweenWarning = 10;
+	double LastClientActivtyTime = 0;
+	double LastServerActivtyTime = 0;
+	double LastWarningTime = 0;
+};
+
+TSharedPtr<IPackageStore> MakeCookOnTheFlyPackageStore(UE::Cook::ICookOnTheFlyServerConnection& CookOnTheFlyServerConnection)
+{
+	return MakeShared<FCookOnTheFlyPackageStore>(CookOnTheFlyServerConnection);
+}
+
+#endif // WITH_COTF
diff --git a/Engine/Source/Runtime/CookOnTheFly/Public/CookOnTheFlyPackageStore.h b/Engine/Source/Runtime/CookOnTheFly/Public/CookOnTheFlyPackageStore.h
new file mode 100644
--- /dev/null
+++ b/Engine/Source/Runtime/CookOnTheFly/Public/CookOnTheFlyPackageStore.h
@@ -0,0 +1,69 @@
+// Copyright Epic Games, Inc. All Rights Reserved.
+
+#pragma once
+
+#include "CoreMinimal.h"
+
+#if WITH_COTF
+
+#include "IO/PackageStore.h"
+
+class FIoDispatcher;
+
+namespace UE { namespace Cook
+{
+	class ICookOnTheFlyServerConnection;
+}}
+
+namespace UE { namespace PackageStore { namespace Messaging
+{
+
+struct FPackageStoreData
+{
+	TArray<FPackageStoreEntryResource> CookedPackages;
+	TArray<FPackageId> FailedPackages;
+	int32 TotalCookedPackages = 0;
+	int32 TotalFailedPackages = 0;
+
+	COOKONTHEFLY_API friend FArchive& operator<<(FArchive& Ar, FPackageStoreData& PackageStoreData);
+};
+
+struct FCookPackageRequest
+{
+	FName PackageName;
+	
+	COOKONTHEFLY_API friend FArchive& operator<<(FArchive& Ar, FCookPackageRequest& Request);
+};
+
+struct FCookPackageResponse
+{
+	EPackageStoreEntryStatus Status = EPackageStoreEntryStatus::None;
+
+	COOKONTHEFLY_API friend FArchive& operator<<(FArchive& Ar, FCookPackageResponse& Response);
+};
+
+struct FPackagesCookedMessage
+{
+	FPackageStoreData PackageStoreData;
+
+	friend FArchive& operator<<(FArchive& Ar, FPackagesCookedMessage& PackagesCookedMessage)
+	{
+		return Ar << PackagesCookedMessage.PackageStoreData;
+	}
+};
+
+struct FGetCookedPackagesResponse
+{
+	FPackageStoreData PackageStoreData;
+	
+	friend FArchive& operator<<(FArchive& Ar, FGetCookedPackagesResponse& GetCookedPackagesResponse)
+	{
+		return Ar << GetCookedPackagesResponse.PackageStoreData;
+	}
+};
+
+}}} // namesapce UE::PackageStore::Messaging
+
+COOKONTHEFLY_API TSharedPtr<IPackageStore> MakeCookOnTheFlyPackageStore(UE::Cook::ICookOnTheFlyServerConnection& CookOnTheFlyServerConnection);
+
+#endif // WITH_COTF
diff --git a/Engine/Source/Runtime/Core/Internal/IO/IoDispatcherFileBackend.h b/Engine/Source/Runtime/Core/Internal/IO/IoDispatcherFileBackend.h
--- a/Engine/Source/Runtime/Core/Internal/IO/IoDispatcherFileBackend.h
+++ b/Engine/Source/Runtime/Core/Internal/IO/IoDispatcherFileBackend.h
@@ -46,10 +46,11 @@ public:
 	const FGuid& GetEncryptionKeyGuid() const { return ContainerFile.EncryptionKeyGuid; }
 	void SetEncryptionKey(const FAES::FAESKey& Key) { ContainerFile.EncryptionKey = Key; }
 	const FAES::FAESKey& GetEncryptionKey() const { return ContainerFile.EncryptionKey; }
+	TIoStatusOr<FIoContainerHeader> ReadContainerHeader() const;
 
 private:
 	const FIoOffsetAndLength* FindChunkInternal(const FIoChunkId& ChunkId) const;
-
+	
 	IPlatformFileIoStore& PlatformImpl;
 	
 	struct FPerfectHashMap
@@ -101,8 +102,8 @@ public:
 	FFileIoStore(TUniquePtr<IPlatformFileIoStore>&& PlatformImpl);
 	~FFileIoStore();
 	void Initialize(TSharedRef<const FIoDispatcherBackendContext> Context) override;
-	TIoStatusOr<FIoContainerId> Mount(const TCHAR* ContainerPath, int32 Order, const FGuid& EncryptionKeyGuid, const FAES::FAESKey& EncryptionKey) override;
-	TIoStatusOr<FIoContainerId> Unmount(const TCHAR* ContainerPath) override;
+	TIoStatusOr<FIoContainerHeader> Mount(const TCHAR* ContainerPath, int32 Order, const FGuid& EncryptionKeyGuid, const FAES::FAESKey& EncryptionKey) override;
+	bool Unmount(const TCHAR* ContainerPath) override;
 	bool Resolve(FIoRequestImpl* Request) override;
 	void CancelIoRequest(FIoRequestImpl* Request) override;
 	void UpdatePriorityForIoRequest(FIoRequestImpl* Request) override;
@@ -110,7 +111,6 @@ public:
 	TIoStatusOr<uint64> GetSizeForChunk(const FIoChunkId& ChunkId) const;
 	FIoRequestImpl* GetCompletedRequests() override;
 	TIoStatusOr<FIoMappedRegion> OpenMapped(const FIoChunkId& ChunkId, const FIoReadOptions& Options) override;
-	void AppendMountedContainers(TSet<FIoContainerId>& OutContainers) override;
 
 	virtual bool Init() override;
 	virtual uint32 Run() override;
diff --git a/Engine/Source/Runtime/Core/Private/IO/IoContainerHeader.cpp b/Engine/Source/Runtime/Core/Private/IO/IoContainerHeader.cpp
new file mode 100644
--- /dev/null
+++ b/Engine/Source/Runtime/Core/Private/IO/IoContainerHeader.cpp
@@ -0,0 +1,37 @@
+// Copyright Epic Games, Inc. All Rights Reserved.
+
+#include "IO/IoContainerHeader.h"
+#include "UObject/NameBatchSerialization.h"
+
+FArchive& operator<<(FArchive& Ar, FIoContainerHeaderPackageRedirect& Redirect)
+{
+	Ar << Redirect.SourcePackageId;
+	Ar << Redirect.TargetPackageId;
+	Ar << Redirect.SourcePackageName;
+
+	return Ar;
+}
+
+FArchive& operator<<(FArchive& Ar, FIoContainerHeader& ContainerHeader)
+{
+	Ar << ContainerHeader.ContainerId;
+	Ar << ContainerHeader.PackageCount;
+	Ar << ContainerHeader.PackageIds;
+	Ar << ContainerHeader.StoreEntries;
+	if (Ar.IsLoading())
+	{
+		ContainerHeader.RedirectsNameMap = LoadNameBatch(Ar);
+	}
+	else
+	{
+#if ALLOW_NAME_BATCH_SAVING
+		SaveNameBatch(ContainerHeader.RedirectsNameMap, Ar);
+#else
+		check(false);
+#endif
+	}
+	Ar << ContainerHeader.CulturePackageMap;
+	Ar << ContainerHeader.PackageRedirects;
+
+	return Ar;
+}
diff --git a/Engine/Source/Runtime/Core/Private/IO/IoDispatcher.cpp b/Engine/Source/Runtime/Core/Private/IO/IoDispatcher.cpp
--- a/Engine/Source/Runtime/Core/Private/IO/IoDispatcher.cpp
+++ b/Engine/Source/Runtime/Core/Private/IO/IoDispatcher.cpp
@@ -274,29 +274,6 @@ public:
 		}
 	}
 
-	FIoStatus Unmount(const TCHAR* ContainerPath)
-	{
-		/*
-		TODO: @Per.Larsson to fixup after unbackported UE5 ioStore changes.
-		TIoStatusOr<FIoContainerId> UnmountedContainerId = FileIoStore.Unmount(ContainerPath);
-
-		if (UnmountedContainerId.IsOk())
-		{
-			if (ContainerUnmountedEvent.IsBound())
-			{
-				FIoDispatcherMountedContainer UnmountedContainer;
-				UnmountedContainer.ContainerId = UnmountedContainerId.ValueOrDie();
-				UnmountedContainer.Environment.InitializeFileEnvironment(ContainerPath);
-
-				ContainerUnmountedEvent.Broadcast(UnmountedContainer);
-			}
-		}
-		*/
-
-		TIoStatusOr<FIoContainerId> UnmountedContainerId;
-		return UnmountedContainerId.Status();
-	}
-
 	bool DoesChunkExist(const FIoChunkId& ChunkId) const
 	{
 		for (const TSharedRef<IIoDispatcherBackend>& Backend : Backends)
@@ -330,26 +307,6 @@ public:
 		}	
 	}
 
-	TSet<FIoContainerId> GetMountedContainers() const
-	{
-		TSet<FIoContainerId> Result;
-		for (const TSharedRef<IIoDispatcherBackend>& Backend : Backends)
-		{
-			Backend->AppendMountedContainers(Result);
-		}
-		return Result;
-	}
-
-	FIoContainerMountedDelegate& OnContainerMounted()
-	{
-		return BackendContext->ContainerMountedDelegate;
-	}
-
-	FIoDispatcher::FIoContainerUnmountedEvent& OnContainerUnmounted()
-	{
-		return ContainerUnmountedEvent;
-	}
-
 	FIoSignatureErrorDelegate& OnSignatureError()
 	{
 		return BackendContext->SignatureErrorDelegate;
@@ -698,11 +655,6 @@ FIoDispatcher::Mount(TSharedRef<IIoDispatcherBackend> Backend)
 	Impl->Mount(Backend);
 }
 
-FIoStatus FIoDispatcher::Unmount(const TCHAR* ContainerPath)
-{
-	return Impl->Unmount(ContainerPath);
-}
-
 FIoBatch
 FIoDispatcher::NewBatch()
 {
@@ -728,30 +680,12 @@ FIoDispatcher::GetSizeForChunk(const FIoChunkId& ChunkId) const
 	return Impl->GetSizeForChunk(ChunkId);
 }
 
-TSet<FIoContainerId>
-FIoDispatcher::GetMountedContainers() const
-{
-	return Impl->GetMountedContainers();
-}
-
 int64
 FIoDispatcher::GetTotalLoaded() const
 {
 	return Impl->GetTotalLoaded();
 }
 
-FIoContainerMountedDelegate&
-FIoDispatcher::OnContainerMounted()
-{
-	return Impl->OnContainerMounted();
-}
-
-FIoDispatcher::FIoContainerUnmountedEvent&
-FIoDispatcher::OnContainerUnmounted()
-{
-	return Impl->OnContainerUnmounted();
-}
-
 FIoSignatureErrorDelegate&
 FIoDispatcher::OnSignatureError()
 {
diff --git a/Engine/Source/Runtime/Core/Private/IO/IoDispatcherFileBackend.cpp b/Engine/Source/Runtime/Core/Private/IO/IoDispatcherFileBackend.cpp
--- a/Engine/Source/Runtime/Core/Private/IO/IoDispatcherFileBackend.cpp
+++ b/Engine/Source/Runtime/Core/Private/IO/IoDispatcherFileBackend.cpp
@@ -21,6 +21,8 @@
 #include "Templates/Greater.h"
 #include "Containers/Ticker.h"
 #include "Modules/ModuleManager.h"
+#include "IO/IoContainerHeader.h"
+#include "Serialization/MemoryReader.h"
 
 TRACE_DECLARE_INT_COUNTER(IoDispatcherLatencyCircuitBreaks, TEXT("IoDispatcher/LatencyCircuitBreaks"));
 TRACE_DECLARE_INT_COUNTER(IoDispatcherSeekDistanceCircuitBreaks, TEXT("IoDispatcher/SeekDistanceCircuitBreaks"));
@@ -821,6 +823,82 @@ IMappedFileHandle* FFileIoStoreReader::GetMappedContainerFileHandle(uint64 TocOf
 	return new FMappedFileProxy(Partition.MappedFileHandle.Get(), Partition.FileSize);
 }
 
+TIoStatusOr<FIoContainerHeader> FFileIoStoreReader::ReadContainerHeader() const
+{
+	LLM_SCOPE(ELLMTag::AsyncLoading);
+	TRACE_CPUPROFILER_EVENT_SCOPE(ReadContainerHeader);
+	FIoChunkId HeaderChunkId = CreateIoChunkId(ContainerId.Value(), 0, EIoChunkType::ContainerHeader);
+	const FIoOffsetAndLength* OffsetAndLength = FindChunkInternal(HeaderChunkId);
+	if (!OffsetAndLength)
+	{
+		return FIoStatus(EIoErrorCode::NotFound);
+	}
+	
+	const uint64 CompressionBlockSize = ContainerFile.CompressionBlockSize;
+	const uint64 Offset = OffsetAndLength->GetOffset();
+	const uint64 Size = OffsetAndLength->GetLength();
+	const uint64 RequestEndOffset = Offset + Size;
+	const int32 RequestBeginBlockIndex = int32(Offset / CompressionBlockSize);
+	const int32 RequestEndBlockIndex = int32((RequestEndOffset - 1) / CompressionBlockSize);
+	
+	// Assumes that the container header is uncompressed and placed in its own blocks in the same partition without padding
+	const FIoStoreTocCompressedBlockEntry* CompressionBlockEntry = &ContainerFile.CompressionBlocks[RequestBeginBlockIndex];
+	const int32 PartitionIndex = int32(CompressionBlockEntry->GetOffset() / ContainerFile.PartitionSize);
+	const FFileIoStoreContainerFilePartition& Partition = ContainerFile.Partitions[PartitionIndex];
+	const uint64 RawOffset = CompressionBlockEntry->GetOffset() % ContainerFile.PartitionSize;
+
+	FIoBuffer IoBuffer(Align(Size, FAES::AESBlockSize));
+	IPlatformFile& Ipf = FPlatformFileManager::Get().GetPlatformFile();
+	TUniquePtr<IFileHandle> ContainerFileHandle(Ipf.OpenRead(*Partition.FilePath));
+	if (!ContainerFileHandle)
+	{
+		return FIoStatus(EIoErrorCode::FileOpenFailed);
+	}
+	{
+		TRACE_CPUPROFILER_EVENT_SCOPE(ReadFromContainerFile);
+		if (!ContainerFileHandle->Seek(RawOffset))
+		{
+			return FIoStatus(EIoErrorCode::ReadError);
+		}
+		if (!ContainerFileHandle->Read(IoBuffer.Data(), IoBuffer.DataSize()))
+		{
+			return FIoStatus(EIoErrorCode::ReadError);
+		}
+	}
+
+	const bool bSigned = EnumHasAnyFlags(ContainerFile.ContainerFlags, EIoContainerFlags::Signed);
+	const bool bEncrypted = ContainerFile.EncryptionKey.IsValid();
+	if (bSigned || bEncrypted)
+	{
+		uint8* BlockData = IoBuffer.Data();
+		for (int32 CompressedBlockIndex = RequestBeginBlockIndex; CompressedBlockIndex <= RequestEndBlockIndex; ++CompressedBlockIndex)
+		{
+			CompressionBlockEntry = &ContainerFile.CompressionBlocks[CompressedBlockIndex];
+			check(ContainerFile.CompressionMethods[CompressionBlockEntry->GetCompressionMethodIndex()].IsNone());
+			const uint64 BlockSize = Align(CompressionBlockEntry->GetCompressedSize(), FAES::AESBlockSize);
+			if (bSigned)
+			{
+				const FSHAHash& SignatureHash = ContainerFile.BlockSignatureHashes[CompressedBlockIndex];
+				FSHAHash BlockHash;
+				FSHA1::HashBuffer(BlockData, BlockSize, BlockHash.Hash);
+				if (SignatureHash != BlockHash)
+				{
+					return FIoStatus(EIoErrorCode::SignatureError);
+				}
+			}
+			if (bEncrypted)
+			{
+				FAES::DecryptData(BlockData, uint32(BlockSize), ContainerFile.EncryptionKey);
+			}
+			BlockData += BlockSize;
+		}
+	}
+	FMemoryReaderView Ar(MakeArrayView(IoBuffer.Data(), static_cast<int32>(IoBuffer.DataSize())));
+	FIoContainerHeader ContainerHeader;
+	Ar << ContainerHeader;
+	return ContainerHeader;
+}
+
 FFileIoStoreResolvedRequest::FFileIoStoreResolvedRequest(
 	FIoRequestImpl& InDispatcherRequest,
 	const FFileIoStoreContainerFile& InContainerFile,
@@ -1108,7 +1186,7 @@ void FFileIoStore::Initialize(TSharedRef<const FIoDispatcherBackendContext> InCo
 	FFileIoStats::SetFileIoStoreThreadId(Thread ? Thread->GetThreadID() : 0);
 }
 
-TIoStatusOr<FIoContainerId> FFileIoStore::Mount(const TCHAR* ContainerPath, int32 Order, const FGuid& EncryptionKeyGuid, const FAES::FAESKey& EncryptionKey)
+TIoStatusOr<FIoContainerHeader> FFileIoStore::Mount(const TCHAR* ContainerPath, int32 Order, const FGuid& EncryptionKeyGuid, const FAES::FAESKey& EncryptionKey)
 {
 	TUniquePtr<FFileIoStoreReader> Reader(new FFileIoStoreReader(*PlatformImpl));
 	FIoStatus IoStatus = Reader->Initialize(ContainerPath, Order);
@@ -1130,8 +1208,18 @@ TIoStatusOr<FIoContainerId> FFileIoStore::Mount(const TCHAR* ContainerPath, int3
 		}
 	}
 
+	TIoStatusOr<FIoContainerHeader> ContainerHeaderReadResult = Reader->ReadContainerHeader();
+	FIoContainerHeader ContainerHeader;
+	if (ContainerHeaderReadResult.IsOk())
+	{
+		ContainerHeader = ContainerHeaderReadResult.ConsumeValueOrDie();
+	}
+	else if (EIoErrorCode::NotFound != ContainerHeaderReadResult.Status().GetErrorCode())
+	{
+		return ContainerHeaderReadResult;
+	}
+
 	int32 InsertionIndex;
-	FIoContainerId ContainerId = Reader->GetContainerId();
 	{
 		FWriteScopeLock _(IoStoreReadersLock);
 		InsertionIndex = Algo::UpperBound(IoStoreReaders, Reader, [](const TUniquePtr<FFileIoStoreReader>& A, const TUniquePtr<FFileIoStoreReader>& B)
@@ -1145,14 +1233,11 @@ TIoStatusOr<FIoContainerId> FFileIoStore::Mount(const TCHAR* ContainerPath, int3
 		IoStoreReaders.Insert(MoveTemp(Reader), InsertionIndex);
 		UE_LOG(LogIoDispatcher, Display, TEXT("Mounting container '%s' in location slot %d"), ContainerPath, InsertionIndex);
 	}
-	if (BackendContext && BackendContext->ContainerMountedDelegate.IsBound())
-	{
-		BackendContext->ContainerMountedDelegate.Broadcast(ContainerId);
-	}
-	return ContainerId;
+
+	return ContainerHeader;
 }
 
-TIoStatusOr<FIoContainerId> FFileIoStore::Unmount(const TCHAR* ContainerPath)
+bool FFileIoStore::Unmount(const TCHAR* ContainerPath)
 {
 	FWriteScopeLock _(IoStoreReadersLock);
 	
@@ -1173,13 +1258,13 @@ TIoStatusOr<FIoContainerId> FFileIoStore::Unmount(const TCHAR* ContainerPath)
 			FIoContainerId ContainerId = IoStoreReaders[Idx]->GetContainerId();
 			IoStoreReaders.RemoveAt(Idx);
 			
-			return TIoStatusOr<FIoContainerId>(ContainerId);
+			return true;
 		}
 	}
 
 	UE_LOG(LogIoDispatcher, Display, TEXT("Failed to unmount container '%s'"), *FPaths::GetBaseFilename(ContainerPath));
 	
-	return FIoStatus(EIoErrorCode::NotFound);
+	return false;
 }
 
 bool FFileIoStore::Resolve(FIoRequestImpl* Request)
@@ -1651,15 +1736,6 @@ FIoRequestImpl* FFileIoStore::GetCompletedRequests()
 	return Result;
 }
 
-void FFileIoStore::AppendMountedContainers(TSet<FIoContainerId>& OutContainers)
-{
-	FReadScopeLock _(IoStoreReadersLock);
-	for (TUniquePtr<FFileIoStoreReader>& Reader : IoStoreReaders)
-	{
-		OutContainers.Add(Reader->GetContainerId());
-	}
-}
-
 TIoStatusOr<FIoMappedRegion> FFileIoStore::OpenMapped(const FIoChunkId& ChunkId, const FIoReadOptions& Options)
 {
 	if (!FPlatformProperties::SupportsMemoryMappedFiles())
diff --git a/Engine/Source/Runtime/Core/Private/Misc/CoreDelegates.cpp b/Engine/Source/Runtime/Core/Private/Misc/CoreDelegates.cpp
--- a/Engine/Source/Runtime/Core/Private/Misc/CoreDelegates.cpp
+++ b/Engine/Source/Runtime/Core/Private/Misc/CoreDelegates.cpp
@@ -196,6 +196,8 @@ FCoreDelegates::FPreloadPackageShaderMaps FCoreDelegates::PreloadPackageShaderMa
 FCoreDelegates::FReleasePreloadedPackageShaderMaps FCoreDelegates::ReleasePreloadedPackageShaderMaps;
 FCoreDelegates::FOnLogVerbosityChanged FCoreDelegates::OnLogVerbosityChanged;
 
+FCoreDelegates::FCreatePackageStore FCoreDelegates::CreatePackageStore;
+
 FCoreDelegates::FApplicationNetworkInitializationChanged FCoreDelegates::ApplicationNetworkInitializationChanged;
 
 PRAGMA_DISABLE_DEPRECATION_WARNINGS
diff --git a/Engine/Source/Runtime/Core/Private/Serialization/MappedName.cpp b/Engine/Source/Runtime/Core/Private/Serialization/MappedName.cpp
new file mode 100644
--- /dev/null
+++ b/Engine/Source/Runtime/Core/Private/Serialization/MappedName.cpp
@@ -0,0 +1,18 @@
+// Copyright Epic Games, Inc. All Rights Reserved.
+
+#include "Serialization/MappedName.h"
+#include "Serialization/Archive.h"
+#include "UObject/NameBatchSerialization.h"
+
+FArchive& operator<<(FArchive& Ar, FMappedName& MappedName)
+{
+	Ar << MappedName.Index << MappedName.Number;
+
+	return Ar;
+}
+
+void FNameMap::Load(FArchive& Ar, FMappedName::EType InNameMapType)
+{
+	NameEntries = LoadNameBatch(Ar);
+	NameMapType = InNameMapType;
+}
diff --git a/Engine/Source/Runtime/Core/Public/IO/IoContainerHeader.h b/Engine/Source/Runtime/Core/Public/IO/IoContainerHeader.h
new file mode 100644
--- /dev/null
+++ b/Engine/Source/Runtime/Core/Public/IO/IoContainerHeader.h
@@ -0,0 +1,70 @@
+// Copyright Epic Games, Inc. All Rights Reserved.
+
+#pragma once
+
+#include "IO/PackageId.h"
+#include "IO/IoContainerId.h"
+#include "Serialization/MappedName.h"
+
+class FArchive;
+
+/**
+ * Package store entry array view.
+ */
+template<typename T>
+class TFilePackageStoreEntryCArrayView
+{
+	const uint32 ArrayNum = 0;
+	const uint32 OffsetToDataFromThis = 0;
+
+public:
+	inline uint32 Num() const { return ArrayNum; }
+
+	inline const T* Data() const { return (T*)((char*)this + OffsetToDataFromThis); }
+	inline T* Data() { return (T*)((char*)this + OffsetToDataFromThis); }
+
+	inline const T* begin() const { return Data(); }
+	inline T* begin() { return Data(); }
+
+	inline const T* end() const { return Data() + ArrayNum; }
+	inline T* end() { return Data() + ArrayNum; }
+
+	inline const T& operator[](uint32 Index) const { return Data()[Index]; }
+	inline T& operator[](uint32 Index) { return Data()[Index]; }
+};
+
+/**
+ * File based package store entry
+ */
+struct FFilePackageStoreEntry
+{
+	int32 ExportCount;
+	int32 ExportBundleCount;
+	TFilePackageStoreEntryCArrayView<FPackageId> ImportedPackages;
+	TFilePackageStoreEntryCArrayView<FSHAHash> ShaderMapHashes;
+};
+
+struct FIoContainerHeaderPackageRedirect
+{
+	FPackageId SourcePackageId;
+	FPackageId TargetPackageId;
+	FMappedName SourcePackageName;
+
+	CORE_API friend FArchive& operator<<(FArchive& Ar, FIoContainerHeaderPackageRedirect& PackageRedirect);
+};
+
+using FSourceToLocalizedPackageIdMap = TArray<FIoContainerHeaderPackageRedirect>;
+using FCulturePackageMap = TMap<FString, FSourceToLocalizedPackageIdMap>;
+
+struct FIoContainerHeader
+{
+	FIoContainerId ContainerId;
+	uint32 PackageCount = 0;
+	TArray<FPackageId> PackageIds;
+	TArray<uint8> StoreEntries; //FPackageStoreEntry[PackageCount]
+	TArray<FNameEntryId> RedirectsNameMap;
+	FCulturePackageMap CulturePackageMap;
+	TArray<FIoContainerHeaderPackageRedirect> PackageRedirects;
+
+	CORE_API friend FArchive& operator<<(FArchive& Ar, FIoContainerHeader& ContainerHeader);
+};
diff --git a/Engine/Source/Runtime/Core/Public/IO/IoDispatcher.h b/Engine/Source/Runtime/Core/Public/IO/IoDispatcher.h
--- a/Engine/Source/Runtime/Core/Public/IO/IoDispatcher.h
+++ b/Engine/Source/Runtime/Core/Public/IO/IoDispatcher.h
@@ -941,7 +941,6 @@ public:
 	CORE_API						~FIoDispatcher();
 
 	CORE_API void					Mount(TSharedRef<IIoDispatcherBackend> Backend);
-	CORE_API FIoStatus				Unmount(const TCHAR* ContainerPath);
 
 	CORE_API FIoBatch				NewBatch();
 
@@ -950,13 +949,10 @@ public:
 	// Polling methods
 	CORE_API bool					DoesChunkExist(const FIoChunkId& ChunkId) const;
 	CORE_API TIoStatusOr<uint64>	GetSizeForChunk(const FIoChunkId& ChunkId) const;
-	CORE_API TSet<FIoContainerId>	GetMountedContainers() const;
 	CORE_API int64					GetTotalLoaded() const;
 
 
 	// Events
-	CORE_API FIoContainerMountedDelegate& OnContainerMounted();
-	CORE_API FIoContainerUnmountedEvent& OnContainerUnmounted();
 	CORE_API FIoSignatureErrorDelegate& OnSignatureError();
 
 	FIoDispatcher(const FIoDispatcher&) = default;
diff --git a/Engine/Source/Runtime/Core/Public/IO/IoDispatcherBackend.h b/Engine/Source/Runtime/Core/Public/IO/IoDispatcherBackend.h
--- a/Engine/Source/Runtime/Core/Public/IO/IoDispatcherBackend.h
+++ b/Engine/Source/Runtime/Core/Public/IO/IoDispatcherBackend.h
@@ -7,6 +7,8 @@
 #include "Misc/Optional.h"
 #include "ProfilingDebugging/TagTrace.h"
 
+struct FIoContainerHeader;
+
 class FIoRequestImpl
 {
 public:
@@ -87,7 +89,6 @@ DECLARE_DELEGATE(FWakeUpIoDispatcherThreadDelegate);
 struct FIoDispatcherBackendContext
 {
 	FWakeUpIoDispatcherThreadDelegate WakeUpDispatcherThreadDelegate;
-	FIoContainerMountedDelegate ContainerMountedDelegate;
 	FIoSignatureErrorDelegate SignatureErrorDelegate;
 	bool bIsMultiThreaded;
 };
@@ -102,14 +103,13 @@ struct IIoDispatcherBackend
 	virtual TIoStatusOr<uint64> GetSizeForChunk(const FIoChunkId& ChunkId) const = 0;
 	virtual FIoRequestImpl* GetCompletedRequests() = 0;
 	virtual TIoStatusOr<FIoMappedRegion> OpenMapped(const FIoChunkId& ChunkId, const FIoReadOptions& Options) = 0;
-	virtual void AppendMountedContainers(TSet<FIoContainerId>& OutContainers) = 0;
 };
 
 struct IIoDispatcherFileBackend
 	: public IIoDispatcherBackend
 {
-	virtual TIoStatusOr<FIoContainerId> Mount(const TCHAR* ContainerPath, int32 Order, const FGuid& EncryptionKeyGuid, const FAES::FAESKey& EncryptionKey) = 0;
-	virtual TIoStatusOr<FIoContainerId> Unmount(const TCHAR* ContainerPath) = 0;
+	virtual TIoStatusOr<FIoContainerHeader> Mount(const TCHAR* ContainerPath, int32 Order, const FGuid& EncryptionKeyGuid, const FAES::FAESKey& EncryptionKey) = 0;
+	virtual bool Unmount(const TCHAR* ContainerPath) = 0;
 };
 
 CORE_API TSharedRef<IIoDispatcherFileBackend> CreateIoDispatcherFileBackend();
diff --git a/Engine/Source/Runtime/Core/Public/IO/PackageStore.h b/Engine/Source/Runtime/Core/Public/IO/PackageStore.h
--- a/Engine/Source/Runtime/Core/Public/IO/PackageStore.h
+++ b/Engine/Source/Runtime/Core/Public/IO/PackageStore.h
@@ -186,9 +186,14 @@ public:
 	/* Destructor. */
 	virtual ~IPackageStore() { }
 
-	/* Initialize the package store. */
 	virtual void Initialize() = 0;
 
+	/** Lock the package store for reading */
+	virtual void Lock() = 0;
+
+	/** Unlock the package store */
+	virtual void Unlock() = 0;
+
 	/* Returns whether the package exists. */
 	virtual bool DoesPackageExist(FPackageId PackageId) = 0;
 
@@ -200,4 +205,21 @@ public:
 
 	/* Returns the redirected package ID and source package name for the specified package ID if it's being redirected. */
 	virtual bool GetPackageRedirectInfo(FPackageId PackageId, FName& OutSourcePackageName, FPackageId& OutRedirectedToPackageId) = 0;
+
+	/* Event broadcasted when pending entries are completed and added to the package store */
+	DECLARE_EVENT(IPackageStore, FEntriesAddedEvent);
+	virtual FEntriesAddedEvent& OnPendingEntriesAdded() = 0;
 };
+
+class FPackageStoreBase
+	: public IPackageStore
+{
+public:
+	virtual FEntriesAddedEvent& OnPendingEntriesAdded() override
+	{
+		return PendingEntriesAdded;
+	}
+
+protected:
+	FEntriesAddedEvent PendingEntriesAdded;
+};
\ No newline at end of file
diff --git a/Engine/Source/Runtime/Core/Public/Misc/CoreDelegates.h b/Engine/Source/Runtime/Core/Public/Misc/CoreDelegates.h
--- a/Engine/Source/Runtime/Core/Public/Misc/CoreDelegates.h
+++ b/Engine/Source/Runtime/Core/Public/Misc/CoreDelegates.h
@@ -628,6 +628,9 @@ public:
 	DECLARE_MULTICAST_DELEGATE_ThreeParams(FOnLogVerbosityChanged, const FLogCategoryName& /* CategoryName */, ELogVerbosity::Type /* OldVerbosity */, ELogVerbosity::Type /* NewVerbosity */);
 	static FOnLogVerbosityChanged OnLogVerbosityChanged;
 
+	DECLARE_DELEGATE_RetVal(TSharedPtr<class IPackageStore>, FCreatePackageStore);
+	static FCreatePackageStore CreatePackageStore;
+
 private:
 
 	// Callbacks for hotfixes
diff --git a/Engine/Source/Runtime/Core/Public/Serialization/MappedName.h b/Engine/Source/Runtime/Core/Public/Serialization/MappedName.h
new file mode 100644
--- /dev/null
+++ b/Engine/Source/Runtime/Core/Public/Serialization/MappedName.h
@@ -0,0 +1,147 @@
+// Copyright Epic Games, Inc. All Rights Reserved.
+
+#pragma once
+
+#include "CoreMinimal.h"
+
+/**
+ * Index and name number into a name map.
+ */
+class FMappedName
+{
+	static constexpr uint32 InvalidIndex = ~uint32(0);
+	static constexpr uint32 IndexBits = 30u;
+	static constexpr uint32 IndexMask = (1u << IndexBits) - 1u;
+	static constexpr uint32 TypeMask = ~IndexMask;
+	static constexpr uint32 TypeShift = IndexBits;
+
+public:
+	enum class EType
+	{
+		Package,
+		Container,
+		Global
+	};
+
+	inline FMappedName() = default;
+
+	static inline FMappedName Create(const uint32 InIndex, const uint32 InNumber, EType InType)
+	{
+		check(InIndex <= MAX_int32);
+		return FMappedName((uint32(InType) << TypeShift) | InIndex, InNumber);
+	}
+
+	static inline FMappedName FromMinimalName(const FMinimalName& MinimalName)
+	{
+		return *reinterpret_cast<const FMappedName*>(&MinimalName);
+	}
+
+	static inline bool IsResolvedToMinimalName(const FMinimalName& MinimalName)
+	{
+		// Not completely safe, relies on that no FName will have its Index and Number equal to Max_uint32
+		const FMappedName MappedName = FromMinimalName(MinimalName);
+		return MappedName.IsValid();
+	}
+
+	static inline FName SafeMinimalNameToName(const FMinimalName& MinimalName)
+	{
+		return IsResolvedToMinimalName(MinimalName) ? MinimalNameToName(MinimalName) : NAME_None;
+	}
+
+	inline FMinimalName ToUnresolvedMinimalName() const
+	{
+		return *reinterpret_cast<const FMinimalName*>(this);
+	}
+
+	inline bool IsValid() const
+	{
+		return Index != InvalidIndex && Number != InvalidIndex;
+	}
+
+	inline EType GetType() const
+	{
+		return static_cast<EType>(uint32((Index & TypeMask) >> TypeShift));
+	}
+
+	inline bool IsGlobal() const
+	{
+		return ((Index & TypeMask) >> TypeShift) != 0;
+	}
+
+	inline uint32 GetIndex() const
+	{
+		return Index & IndexMask;
+	}
+
+	inline uint32 GetNumber() const
+	{
+		return Number;
+	}
+
+	inline bool operator!=(FMappedName Other) const
+	{
+		return Index != Other.Index || Number != Other.Number;
+	}
+
+	CORE_API friend FArchive& operator<<(FArchive& Ar, FMappedName& MappedName);
+
+private:
+	inline FMappedName(const uint32 InIndex, const uint32 InNumber)
+		: Index(InIndex)
+		, Number(InNumber) { }
+
+	uint32 Index = InvalidIndex;
+	uint32 Number = InvalidIndex;
+};
+
+/*
+ * Maps serialized name entries to names.
+ */
+class FNameMap
+{
+public:
+	inline int32 Num() const
+	{
+		return NameEntries.Num();
+	}
+
+	CORE_API void Load(FArchive& Ar, FMappedName::EType NameMapType);
+
+	FName GetName(const FMappedName& MappedName) const
+	{
+		check(MappedName.GetType() == NameMapType);
+		check(MappedName.GetIndex() < uint32(NameEntries.Num()));
+
+		FNameEntryId NameEntry = NameEntries[MappedName.GetIndex()];
+		return FName::CreateFromDisplayId(NameEntry, MappedName.GetNumber());
+	}
+
+	bool TryGetName(const FMappedName& MappedName, FName& OutName) const
+	{
+		check(MappedName.GetType() == NameMapType);
+
+		uint32 Index = MappedName.GetIndex();
+		if (Index < uint32(NameEntries.Num()))
+		{
+			FNameEntryId NameEntry = NameEntries[MappedName.GetIndex()];
+			OutName = FName::CreateFromDisplayId(NameEntry, MappedName.GetNumber());
+
+			return true;
+		}
+
+		return false;
+	}
+
+	FMinimalName GetMinimalName(const FMappedName& MappedName) const
+	{
+		check(MappedName.GetType() == NameMapType);
+		check(MappedName.GetIndex() < uint32(NameEntries.Num()));
+
+		FNameEntryId NameEntry = NameEntries[MappedName.GetIndex()];
+		return FMinimalName(NameEntry, MappedName.GetNumber());
+	}
+
+private:
+	TArray<FNameEntryId> NameEntries;
+	FMappedName::EType NameMapType = FMappedName::EType::Global;
+};
diff --git a/Engine/Source/Runtime/CoreUObject/CoreUObject.Build.cs b/Engine/Source/Runtime/CoreUObject/CoreUObject.Build.cs
--- a/Engine/Source/Runtime/CoreUObject/CoreUObject.Build.cs
+++ b/Engine/Source/Runtime/CoreUObject/CoreUObject.Build.cs
@@ -24,11 +24,6 @@ public class CoreUObject : ModuleRules
 
 		PrivateDependencyModuleNames.Add("Projects");
         PrivateDependencyModuleNames.Add("Json");
-		
-		if (Target.Configuration != UnrealTargetConfiguration.Shipping && Target.Type != TargetType.Program)
-		{
-			PublicDependencyModuleNames.Add("CookOnTheFly");
-		}
 
 		//@TODO: UE-127233
 		// UnsafeTypeCastWarningLevel = WarningLevel.Warning;
diff --git a/Engine/Source/Runtime/CoreUObject/Private/Serialization/AsyncLoading2.cpp b/Engine/Source/Runtime/CoreUObject/Private/Serialization/AsyncLoading2.cpp
--- a/Engine/Source/Runtime/CoreUObject/Private/Serialization/AsyncLoading2.cpp
+++ b/Engine/Source/Runtime/CoreUObject/Private/Serialization/AsyncLoading2.cpp
@@ -7,7 +7,6 @@
 #include "Serialization/AsyncLoading2.h"
 #include "Serialization/AsyncPackageLoader.h"
 #include "IO/PackageStore.h"
-#include "Serialization/FilePackageStore.h"
 #include "HAL/PlatformFileManager.h"
 #include "HAL/FileManager.h"
 #include "HAL/Event.h"
@@ -17,8 +16,6 @@
 #include "Stats/StatsMisc.h"
 #include "Misc/CoreStats.h"
 #include "HAL/IConsoleManager.h"
-#include "Internationalization/Culture.h"
-#include "Internationalization/Internationalization.h"
 #include "Misc/CoreDelegates.h"
 #include "Misc/CommandLine.h"
 #include "Misc/App.h"
@@ -65,15 +62,6 @@
 #include "HAL/IPlatformFileOpenLogWrapper.h"
 #include "Modules/ModuleManager.h"
 
-#ifndef WITH_COTF
-#define WITH_COTF !UE_BUILD_SHIPPING
-#endif
-
-#if WITH_COTF
-#include "CookOnTheFly.h"
-#include "Serialization/CookOnTheFlyPackageStore.h"
-#endif
-
 #if UE_BUILD_DEVELOPMENT || UE_BUILD_DEBUG
 PRAGMA_DISABLE_OPTIMIZATION
 #endif
@@ -81,39 +69,6 @@ PRAGMA_DISABLE_OPTIMIZATION
 CSV_DECLARE_CATEGORY_MODULE_EXTERN(CORE_API, FileIO);
 CSV_DEFINE_STAT(FileIO, FrameCompletedExportBundleLoadsKB);
 
-FArchive& operator<<(FArchive& Ar, FContainerHeaderPackageRedirect& Redirect)
-{
-	Ar << Redirect.SourcePackageId;
-	Ar << Redirect.TargetPackageId;
-	Ar << Redirect.SourcePackageName;
-
-	return Ar;
-}
-
-FArchive& operator<<(FArchive& Ar, FContainerHeader& ContainerHeader)
-{
-	Ar << ContainerHeader.ContainerId;
-	Ar << ContainerHeader.PackageCount;
-	Ar << ContainerHeader.PackageIds;
-	Ar << ContainerHeader.StoreEntries;
-	if (Ar.IsLoading())
-	{
-		ContainerHeader.RedirectsNameMap = LoadNameBatch(Ar);
-	}
-	else
-	{
-#if ALLOW_NAME_BATCH_SAVING
-		SaveNameBatch(ContainerHeader.RedirectsNameMap, Ar);
-#else
-		check(false);
-#endif
-	}
-	Ar << ContainerHeader.CulturePackageMap;
-	Ar << ContainerHeader.PackageRedirects;
-
-	return Ar;
-}
-
 FArchive& operator<<(FArchive& Ar, FExportBundleEntry& ExportBundleEntry)
 {
 	Ar << ExportBundleEntry.LocalExportIndex;
@@ -1784,7 +1739,7 @@ struct FAsyncPackage2
 	/** Creates GC clusters from loaded objects */
 	EAsyncPackageState::Type CreateClusters(FAsyncLoadingThreadState2& ThreadState);
 
-	void ImportPackagesRecursive(FIoBatch& IoBatch);
+	void ImportPackagesRecursive(FIoBatch& IoBatch, IPackageStore& PackageStore);
 	void StartLoading(FIoBatch& IoBatch);
 
 #if WITH_IOSTORE_IN_EDITOR
@@ -2185,7 +2140,7 @@ private:
 	/** I/O Dispatcher */
 	FIoDispatcher& IoDispatcher;
 
-	TUniquePtr<IPackageStore> PackageStore;
+	TSharedPtr<IPackageStore> PackageStore;
 	FLoadedPackageStore LoadedPackageStore;
 	FGlobalImportStore GlobalImportStore;
 	TUniquePtr<FPackageRequestQueue> PackageRequestQueue;
@@ -2436,11 +2391,6 @@ public:
 		}
 	}
 
-	IPackageStore& GetPackageStore()
-	{
-		return *PackageStore.Get();
-	}
-
 private:
 
 	void SuspendWorkers();
@@ -2644,30 +2594,24 @@ void FAsyncLoadingThread2::InitializeLoading()
 	FBulkDataBase::SetIoDispatcher(&IoDispatcher);
 #endif
 
-#if WITH_COTF
-	if (IsRunningCookOnTheFly())
-	{
-		UE::Cook::ICookOnTheFlyModule& CookOnTheFlyModule = FModuleManager::LoadModuleChecked<UE::Cook::ICookOnTheFlyModule>(TEXT("CookOnTheFly"));
-		PackageStore = MakeCookOnTheFlyPackageStore(
-			IoDispatcher,
-			CookOnTheFlyModule.GetServerConnection(),
-			[this]()
-			{
-				AltZenaphore.NotifyOne();
-			});
-	}
-	else
-#endif
+	UE_CLOG(!FCoreDelegates::CreatePackageStore.IsBound(), LogStreaming, Fatal, TEXT("CreatePackageStore delegate not bound when initializing package loader"));
+	PackageStore = FCoreDelegates::CreatePackageStore.Execute();
+	UE_CLOG(!PackageStore.IsValid(), LogStreaming, Fatal, TEXT("CreatePackageStore delegate returned null when initializing package loader"));
+	PackageStore->Initialize();
+	PackageStore->OnPendingEntriesAdded().AddLambda([this]()
 	{
-		PackageStore = MakeFilePackageStore(IoDispatcher);
-	}
-	
-	PackageRequestQueue = MakeUnique<FPackageRequestQueue>(*PackageStore);
+		AltZenaphore.NotifyOne();
+	});
+	PackageRequestQueue = MakeUnique<FPackageRequestQueue>(*PackageStore.Get());
 
 	FPackageName::DoesPackageExistOverride().BindLambda([this](FName PackageName)
 	{
 		LazyInitializeFromLoadPackage();
-		return PackageStore->DoesPackageExist(FPackageId::FromName(PackageName));
+		
+		PackageStore->Lock();
+		bool bExists = PackageStore->DoesPackageExist(FPackageId::FromName(PackageName));
+		PackageStore->Unlock();
+		return bExists;
 	});
 
 	AsyncThreadReady.Increment();
@@ -2727,6 +2671,8 @@ bool FAsyncLoadingThread2::CreateAsyncPackagesFromQueue(FAsyncLoadingThreadState
 
 	FIoBatch IoBatch = IoDispatcher.NewBatch();
 
+	PackageStore->Lock();
+
 	for (;;)
 	{
 		int32 NumDequeued = 0;
@@ -2767,7 +2713,7 @@ bool FAsyncLoadingThread2::CreateAsyncPackagesFromQueue(FAsyncLoadingThreadState
 				{
 					{
 						TRACE_CPUPROFILER_EVENT_SCOPE(ImportPackages);
-						Package->ImportPackagesRecursive(IoBatch);
+						Package->ImportPackagesRecursive(IoBatch, *PackageStore.Get());
 					}
 
 					if (bInserted)
@@ -2788,6 +2734,8 @@ bool FAsyncLoadingThread2::CreateAsyncPackagesFromQueue(FAsyncLoadingThreadState
 		}
 	}
 
+	PackageStore->Unlock();
+
 	IoBatch.Issue();
 	
 	return bPackagesCreated;
@@ -3296,7 +3244,7 @@ void FGlobalImportStore::FindAllScriptObjects()
 		ScriptObjects.Num(), (float)ScriptObjects.GetAllocatedSize() / 1024.f);
 }
 
-void FAsyncPackage2::ImportPackagesRecursive(FIoBatch& IoBatch)
+void FAsyncPackage2::ImportPackagesRecursive(FIoBatch& IoBatch, IPackageStore& PackageStore)
 {
 	if (AsyncPackageLoadingState >= EAsyncPackageLoadingState2::ImportPackages)
 	{
@@ -3317,7 +3265,6 @@ void FAsyncPackage2::ImportPackagesRecursive(FIoBatch& IoBatch)
 
 	int32 ImportedPackageIndex = 0;
 
-	IPackageStore& PackageStore = AsyncLoadingThread.GetPackageStore();
 	Data.ImportedAsyncPackages = MakeArrayView(Data.ImportedAsyncPackages.GetData(), Data.ImportedPackageIds.Num());
 	for (const FPackageId& ImportedPackageId : Data.ImportedPackageIds)
 	{
@@ -3382,7 +3329,7 @@ void FAsyncPackage2::ImportPackagesRecursive(FIoBatch& IoBatch)
 
 		if (bInserted)
 		{
-			ImportedPackage->ImportPackagesRecursive(IoBatch);
+			ImportedPackage->ImportPackagesRecursive(IoBatch, PackageStore);
 			ImportedPackage->StartLoading(IoBatch);
 		}
 	}
@@ -5039,7 +4986,6 @@ void FAsyncLoadingThread2::LazyInitializeFromLoadPackage()
 		GlobalImportStore.Initialize(IoDispatcher);
 	}
 
-	PackageStore->Initialize();
 	LoadedPackageStore.Initialize(*PackageStore.Get());
 }
 
@@ -5495,7 +5441,7 @@ FAsyncPackage2::FAsyncPackage2(
 , Data(InData)
 , AsyncLoadingThread(InAsyncLoadingThread)
 , GraphAllocator(InGraphAllocator)
-, ImportStore(AsyncLoadingThread.GetPackageStore(), AsyncLoadingThread.GlobalImportStore, AsyncLoadingThread.LoadedPackageStore, Desc, Data.ImportedPackageIds)
+, ImportStore(*AsyncLoadingThread.PackageStore.Get(), AsyncLoadingThread.GlobalImportStore, AsyncLoadingThread.LoadedPackageStore, Desc, Data.ImportedPackageIds)
 {
 	TRACE_LOADTIME_NEW_ASYNC_PACKAGE(this);
 	ImportStore.AddPackageReferences();
@@ -5922,7 +5868,9 @@ int32 FAsyncLoadingThread2::LoadPackage(const FPackagePath& InPackagePath, FName
 		? MakeUnique<FLoadPackageAsyncDelegate>(MoveTemp(InCompletionDelegate))
 		: TUniquePtr<FLoadPackageAsyncDelegate>();
 
+	PackageStore->Lock();
 	PackageRequestQueue->Enqueue(FPackageRequest::Create(RequestId, InPackagePriority, PackageNameToLoad, PackageIdToLoad, InCustomName, MoveTemp(CompletionDelegate)));
+	PackageStore->Unlock();
 	++QueuedPackagesCounter;
 
 	AltZenaphore.NotifyOne();
diff --git a/Engine/Source/Runtime/CoreUObject/Private/Serialization/CookOnTheFlyPackageStore.cpp b/Engine/Source/Runtime/CoreUObject/Private/Serialization/CookOnTheFlyPackageStore.cpp
deleted file mode 100644
diff --git a/Engine/Source/Runtime/CoreUObject/Private/Serialization/EditorDerivedDataIoStore.cpp b/Engine/Source/Runtime/CoreUObject/Private/Serialization/EditorDerivedDataIoStore.cpp
--- a/Engine/Source/Runtime/CoreUObject/Private/Serialization/EditorDerivedDataIoStore.cpp
+++ b/Engine/Source/Runtime/CoreUObject/Private/Serialization/EditorDerivedDataIoStore.cpp
@@ -158,7 +158,6 @@ public:
 	TIoStatusOr<uint64> GetSizeForChunk(const FIoChunkId& ChunkId) const final;
 	FIoRequestImpl* GetCompletedRequests() final;
 	TIoStatusOr<FIoMappedRegion> OpenMapped(const FIoChunkId& ChunkId, const FIoReadOptions& Options) final;
-	void AppendMountedContainers(TSet<FIoContainerId>& OutContainers) final;
 
 	FIoChunkId AddData(const Private::FEditorDerivedData* Data) final;
 	const Private::FEditorDerivedData* RemoveData(const FIoChunkId& ChunkId) final;
@@ -248,10 +247,6 @@ TIoStatusOr<FIoMappedRegion> FEditorDerivedDataIoStore::OpenMapped(const FIoChun
 	return FIoStatus(EIoErrorCode::NotFound);
 }
 
-void FEditorDerivedDataIoStore::AppendMountedContainers(TSet<FIoContainerId>& OutContainers)
-{
-}
-
 FIoChunkId FEditorDerivedDataIoStore::AddData(const Private::FEditorDerivedData* Data)
 {
 	FWriteScopeLock WriteLock(DataLock);
diff --git a/Engine/Source/Runtime/CoreUObject/Private/Serialization/FilePackageStore.cpp b/Engine/Source/Runtime/CoreUObject/Private/Serialization/FilePackageStore.cpp
deleted file mode 100644
diff --git a/Engine/Source/Runtime/CoreUObject/Private/Serialization/MappedName.cpp b/Engine/Source/Runtime/CoreUObject/Private/Serialization/MappedName.cpp
deleted file mode 100644
diff --git a/Engine/Source/Runtime/CoreUObject/Public/Serialization/AsyncLoading2.h b/Engine/Source/Runtime/CoreUObject/Public/Serialization/AsyncLoading2.h
--- a/Engine/Source/Runtime/CoreUObject/Public/Serialization/AsyncLoading2.h
+++ b/Engine/Source/Runtime/CoreUObject/Public/Serialization/AsyncLoading2.h
@@ -11,38 +11,12 @@
 #include "IO/PackageId.h"
 #include "Serialization/Archive.h"
 #include "Serialization/MappedName.h"
-#include "IO/IoContainerId.h"
 
 class FArchive;
 class IAsyncPackageLoader;
 class FIoDispatcher;
 class IEDLBootNotificationManager;
 
-struct FContainerHeaderPackageRedirect
-{
-	FPackageId SourcePackageId;
-	FPackageId TargetPackageId;
-	FMappedName SourcePackageName;
-
-	COREUOBJECT_API friend FArchive& operator<<(FArchive& Ar, FContainerHeaderPackageRedirect& PackageRedirect);
-};
-
-using FSourceToLocalizedPackageIdMap = TArray<FContainerHeaderPackageRedirect>;
-using FCulturePackageMap = TMap<FString, FSourceToLocalizedPackageIdMap>;
-
-struct FContainerHeader
-{
-	FIoContainerId ContainerId;
-	uint32 PackageCount = 0;
-	TArray<FPackageId> PackageIds;
-	TArray<uint8> StoreEntries; //FPackageStoreEntry[PackageCount]
-	TArray<FNameEntryId> RedirectsNameMap;
-	FCulturePackageMap CulturePackageMap;
-	TArray<FContainerHeaderPackageRedirect> PackageRedirects;
-
-	COREUOBJECT_API friend FArchive& operator<<(FArchive& Ar, FContainerHeader& ContainerHeader);
-};
-
 class FPackageImportReference
 {
 public:
diff --git a/Engine/Source/Runtime/CoreUObject/Public/Serialization/CookOnTheFlyPackageStore.h b/Engine/Source/Runtime/CoreUObject/Public/Serialization/CookOnTheFlyPackageStore.h
deleted file mode 100644
diff --git a/Engine/Source/Runtime/CoreUObject/Public/Serialization/FilePackageStore.h b/Engine/Source/Runtime/CoreUObject/Public/Serialization/FilePackageStore.h
deleted file mode 100644
diff --git a/Engine/Source/Runtime/CoreUObject/Public/Serialization/MappedName.h b/Engine/Source/Runtime/CoreUObject/Public/Serialization/MappedName.h
deleted file mode 100644
diff --git a/Engine/Source/Runtime/PakFile/PakFile.Build.cs b/Engine/Source/Runtime/PakFile/PakFile.Build.cs
--- a/Engine/Source/Runtime/PakFile/PakFile.Build.cs
+++ b/Engine/Source/Runtime/PakFile/PakFile.Build.cs
@@ -7,6 +7,7 @@ public class PakFile : ModuleRules
 	public PakFile(ReadOnlyTargetRules Target) : base(Target)
 	{
 		PrivateDependencyModuleNames.Add("Core");
+		PrivateDependencyModuleNames.Add("CoreUObject");
 		PublicDependencyModuleNames.Add("RSA");
 	}
 }
diff --git a/Engine/Source/Runtime/PakFile/Private/FilePackageStore.cpp b/Engine/Source/Runtime/PakFile/Private/FilePackageStore.cpp
new file mode 100644
--- /dev/null
+++ b/Engine/Source/Runtime/PakFile/Private/FilePackageStore.cpp
@@ -0,0 +1,210 @@
+// Copyright Epic Games, Inc. All Rights Reserved.
+#include "FilePackageStore.h"
+#include "IO/IoContainerId.h"
+#include "Internationalization/Culture.h"
+#include "Internationalization/Internationalization.h"
+#include "Misc/CommandLine.h"
+#include "IO/IoContainerHeader.h"
+#include "Misc/ScopeRWLock.h"
+#include "ProfilingDebugging/CountersTrace.h"
+#include "Internationalization/PackageLocalizationManager.h"
+
+//PRAGMA_DISABLE_OPTIMIZATION
+
+thread_local bool FFilePackageStore::bIsLockedOnThread = false;
+
+FFilePackageStore::FFilePackageStore()
+{
+}
+
+void FFilePackageStore::Initialize()
+{
+	// Setup culture
+	FInternationalization& Internationalization = FInternationalization::Get();
+	FString CurrentCulture = Internationalization.GetCurrentCulture()->GetName();
+	FParse::Value(FCommandLine::Get(), TEXT("CULTURE="), CurrentCulture);
+	CurrentCultureNames = Internationalization.GetPrioritizedCultureNames(CurrentCulture);
+}
+
+void FFilePackageStore::Lock()
+{
+	EntriesLock.ReadLock();
+	if (bNeedsUpdate)
+	{
+		Update();
+	}
+	bIsLockedOnThread = true;
+}
+
+void FFilePackageStore::Unlock()
+{
+	check(bIsLockedOnThread);
+	EntriesLock.ReadUnlock();
+	bIsLockedOnThread = false;
+}
+
+bool FFilePackageStore::DoesPackageExist(FPackageId PackageId)
+{
+	check(bIsLockedOnThread);
+	return PackageId.IsValid() && StoreEntriesMap.Contains(PackageId);
+}
+
+FPackageStoreEntryHandle FFilePackageStore::GetPackageEntryHandle(FPackageId PackageId, const FName& PackageName)
+{
+	check(bIsLockedOnThread);
+	const FFilePackageStoreEntry* FindEntry = StoreEntriesMap.FindRef(PackageId);
+	const uint64 Handle = reinterpret_cast<uint64>(FindEntry);
+	return FPackageStoreEntryHandle::Create(Handle, Handle ? EPackageStoreEntryStatus::Ok : EPackageStoreEntryStatus::Missing);
+}
+
+FPackageStoreEntry FFilePackageStore::GetPackageEntry(FPackageStoreEntryHandle Handle)
+{
+	check(bIsLockedOnThread);
+	check(Handle.IsValid());
+	const FFilePackageStoreEntry* Entry = reinterpret_cast<const FFilePackageStoreEntry*>(Handle.Value());
+	check(Entry);
+	return FPackageStoreEntry
+	{
+		FPackageStoreExportInfo
+		{
+			Entry->ExportCount,
+			Entry->ExportBundleCount
+		},
+		MakeArrayView(Entry->ImportedPackages.Data(), Entry->ImportedPackages.Num()),
+		MakeArrayView(Entry->ShaderMapHashes.Data(), Entry->ShaderMapHashes.Num())
+	};
+}
+
+bool FFilePackageStore::GetPackageRedirectInfo(FPackageId PackageId, FName& OutSourcePackageName, FPackageId& OutRedirectedToPackageId)
+{
+	check(bIsLockedOnThread);
+	TTuple<FName, FPackageId>* FindRedirect = RedirectsPackageMap.Find(PackageId);
+	if (FindRedirect)
+	{
+		OutSourcePackageName = FindRedirect->Get<0>();
+		OutRedirectedToPackageId = FindRedirect->Get<1>();
+
+		UE_LOG(LogCore, Display, TEXT("Redirecting from %s to 0x%llx"), *OutSourcePackageName.ToString(), OutRedirectedToPackageId.Value());
+
+		return true;
+	}
+	else
+	{
+		return false;
+	}
+}
+
+void FFilePackageStore::Mount(const FIoContainerHeader* ContainerHeader, uint32 Order)
+{
+	LLM_SCOPE(ELLMTag::AsyncLoading);
+	FWriteScopeLock _(EntriesLock);
+	MountedContainers.Add({ ContainerHeader, Order });
+	Algo::Sort(MountedContainers, [](const FMountedContainer& A, const FMountedContainer& B)
+		{
+			return A.Order < B.Order;
+		});
+	bNeedsUpdate = true;
+}
+
+void FFilePackageStore::Unmount(const FIoContainerHeader* ContainerHeader)
+{
+	FWriteScopeLock _(EntriesLock);
+	for (auto It = MountedContainers.CreateIterator(); It; ++It)
+	{
+		if (It->ContainerHeader == ContainerHeader)
+		{
+			It.RemoveCurrent();
+			bNeedsUpdate = true;
+			return;
+		}
+	}
+}
+
+void FFilePackageStore::Update()
+{
+	LLM_SCOPE(ELLMTag::AsyncLoading);
+	TRACE_CPUPROFILER_EVENT_SCOPE(UpdateFilePackageStore);
+
+	FScopeLock Lock(&UpdateLock);
+	if (!bNeedsUpdate)
+	{
+		return;
+	}
+
+	StoreEntriesMap.Empty();
+	RedirectsPackageMap.Empty();
+
+	uint32 TotalPackageCount = 0;
+	for (const FMountedContainer& MountedContainer : MountedContainers)
+	{
+		TotalPackageCount += MountedContainer.ContainerHeader->PackageCount;
+	}
+
+	StoreEntriesMap.Reserve(TotalPackageCount);
+	TMap<FPackageId, FName> AllLocalizedPackages;
+	for (const FMountedContainer& MountedContainer : MountedContainers)
+	{
+		const FIoContainerHeader* ContainerHeader = MountedContainer.ContainerHeader;
+		TArrayView<const FFilePackageStoreEntry> ContainerStoreEntries(reinterpret_cast<const FFilePackageStoreEntry*>(ContainerHeader->StoreEntries.GetData()), ContainerHeader->PackageCount);
+		int32 Index = 0;
+		for (const FFilePackageStoreEntry& StoreEntry : ContainerStoreEntries)
+		{
+			const FPackageId& PackageId = ContainerHeader->PackageIds[Index];
+			check(PackageId.IsValid());
+			StoreEntriesMap.FindOrAdd(PackageId, &StoreEntry);
+			++Index;
+		}
+
+		{
+			TRACE_CPUPROFILER_EVENT_SCOPE(LoadPackageStoreLocalization);
+			const FSourceToLocalizedPackageIdMap* LocalizedPackages = nullptr;
+			for (const FString& CultureName : CurrentCultureNames)
+			{
+				LocalizedPackages = ContainerHeader->CulturePackageMap.Find(CultureName);
+				if (LocalizedPackages)
+				{
+					break;
+				}
+			}
+
+			if (LocalizedPackages)
+			{
+				for (const FIoContainerHeaderPackageRedirect& Redirect : *LocalizedPackages)
+				{
+					FNameEntryId NameEntry = ContainerHeader->RedirectsNameMap[Redirect.SourcePackageName.GetIndex()];
+					FName SourcePackageName = FName::CreateFromDisplayId(NameEntry, Redirect.SourcePackageName.GetNumber());
+					RedirectsPackageMap.Emplace(Redirect.SourcePackageId, MakeTuple(SourcePackageName, Redirect.TargetPackageId));
+				}
+			}
+		}
+
+		{
+			TRACE_CPUPROFILER_EVENT_SCOPE(LoadPackageStoreRedirects);
+			for (const FIoContainerHeaderPackageRedirect& Redirect : ContainerHeader->PackageRedirects)
+			{
+				FNameEntryId NameEntry = ContainerHeader->RedirectsNameMap[Redirect.SourcePackageName.GetIndex()];
+				FName SourcePackageName = FName::CreateFromDisplayId(NameEntry, Redirect.SourcePackageName.GetNumber());
+				RedirectsPackageMap.Emplace(Redirect.SourcePackageId, MakeTuple(SourcePackageName, Redirect.TargetPackageId));
+			}
+		}
+	}
+
+	for (const auto& KV : AllLocalizedPackages)
+	{
+		const FPackageId& SourcePackageId = KV.Key;
+		const FName& SourcePackageName = KV.Value;
+		FName LocalizedPackageName = FPackageLocalizationManager::Get().FindLocalizedPackageName(SourcePackageName);
+		if (!LocalizedPackageName.IsNone())
+		{
+			FPackageId RedirectedToPackageId = FPackageId::FromName(LocalizedPackageName);
+			if (StoreEntriesMap.Find(RedirectedToPackageId))
+			{
+				RedirectsPackageMap.Emplace(SourcePackageId, MakeTuple(SourcePackageName, RedirectedToPackageId));
+			}
+		}
+	}
+	
+	bNeedsUpdate = false;
+}
+
+//PRAGMA_ENABLE_OPTIMIZATION
\ No newline at end of file
diff --git a/Engine/Source/Runtime/PakFile/Private/FilePackageStore.h b/Engine/Source/Runtime/PakFile/Private/FilePackageStore.h
new file mode 100644
--- /dev/null
+++ b/Engine/Source/Runtime/PakFile/Private/FilePackageStore.h
@@ -0,0 +1,51 @@
+// Copyright Epic Games, Inc. All Rights Reserved.
+
+#pragma once
+
+#include "CoreMinimal.h"
+#include "IO/PackageId.h"
+#include "IO/PackageStore.h"
+
+struct FIoContainerHeader;
+struct FFilePackageStoreEntry;
+
+/*
+ * File/container based package store.
+ */
+class FFilePackageStore
+	: public FPackageStoreBase
+{
+public:
+	FFilePackageStore();
+	virtual ~FFilePackageStore() = default;
+
+	virtual void Initialize() override;
+	virtual void Lock() override;
+	virtual void Unlock() override;
+	virtual bool DoesPackageExist(FPackageId PackageId) override;
+	virtual FPackageStoreEntryHandle GetPackageEntryHandle(FPackageId PackageId, const FName& PackageName) override;
+	virtual FPackageStoreEntry GetPackageEntry(FPackageStoreEntryHandle Handle) override;
+	virtual bool GetPackageRedirectInfo(FPackageId PackageId, FName& OutSourcePackageName, FPackageId& OutRedirectedToPackageId) override;
+
+	void Mount(const FIoContainerHeader* ContainerHeader, uint32 Order);
+	void Unmount(const FIoContainerHeader* ContainerHeader);
+
+private:
+	struct FMountedContainer
+	{
+		const FIoContainerHeader* ContainerHeader;
+		uint32 Order;
+	};
+
+	void Update();
+
+	FRWLock EntriesLock;
+	FCriticalSection UpdateLock;
+	TArray<FString> CurrentCultureNames;
+	TArray<FMountedContainer> MountedContainers;
+	TMap<FPackageId, const FFilePackageStoreEntry*> StoreEntriesMap;
+	TMap<FPackageId, TTuple<FName, FPackageId>> RedirectsPackageMap;
+	bool bNeedsUpdate = false;
+
+	static thread_local bool bIsLockedOnThread;
+};
diff --git a/Engine/Source/Runtime/PakFile/Private/IPlatformFilePak.cpp b/Engine/Source/Runtime/PakFile/Private/IPlatformFilePak.cpp
--- a/Engine/Source/Runtime/PakFile/Private/IPlatformFilePak.cpp
+++ b/Engine/Source/Runtime/PakFile/Private/IPlatformFilePak.cpp
@@ -37,6 +37,8 @@
 #include "IO/IoDispatcherBackend.h"
 
 #include "ProfilingDebugging/LoadTimeTracker.h"
+#include "IO/IoContainerHeader.h"
+#include "FilePackageStore.h"
 #include "Compression/OodleDataCompression.h"
 
 DEFINE_LOG_CATEGORY(LogPakFile);
@@ -7324,7 +7326,7 @@ bool FPakPlatformFile::Initialize(IPlatformFile* Inner, const TCHAR* CmdLine)
 		FIoDispatcher& IoDispatcher = FIoDispatcher::Get();
 		IoDispatcherFileBackend = CreateIoDispatcherFileBackend();
 		IoDispatcher.Mount(IoDispatcherFileBackend.ToSharedRef());
-		TIoStatusOr<FIoContainerId> IoDispatcherMountStatus = IoDispatcherFileBackend->Mount(*FPaths::ChangeExtension(GlobalUTocPath, TEXT("")), 0, FGuid(), FAES::FAESKey());
+		TIoStatusOr<FIoContainerHeader> IoDispatcherMountStatus = IoDispatcherFileBackend->Mount(*FPaths::ChangeExtension(GlobalUTocPath, TEXT("")), 0, FGuid(), FAES::FAESKey());
 		if (IoDispatcherMountStatus.IsOk())
 		{
 			UE_LOG(LogPakFile, Display, TEXT("Initialized I/O dispatcher"));
@@ -7337,6 +7339,11 @@ bool FPakPlatformFile::Initialize(IPlatformFile* Inner, const TCHAR* CmdLine)
 #endif
 				FPakPlatformFile::BroadcastPakChunkSignatureCheckFailure(FailedData);
 			});
+			FilePackageStore = MakeShared<FFilePackageStore>();
+			FCoreDelegates::CreatePackageStore.BindLambda([this]()
+			{
+				return FilePackageStore;
+			});
 		}
 		else
 		{
@@ -7606,15 +7613,17 @@ bool FPakPlatformFile::Mount(const TCHAR* InPakFilename, uint32 PakOrder, const
 			}
 
 			FString ContainerPath = FPaths::ChangeExtension(InPakFilename, FString());
-			TIoStatusOr<FIoContainerId> IoStatus = IoDispatcherFileBackend->Mount(*ContainerPath, PakOrder, EncryptionKeyGuid, EncryptionKey);
-			if (IoStatus.IsOk())
+			TIoStatusOr<FIoContainerHeader> MountResult = IoDispatcherFileBackend->Mount(*ContainerPath, PakOrder, EncryptionKeyGuid, EncryptionKey);
+			if (MountResult.IsOk())
 			{
 				UE_LOG(LogPakFile, Display, TEXT("Mounted IoStore container \"%s\""), *ContainerPath);
+				Pak->IoContainerHeader = MakeUnique<FIoContainerHeader>(MountResult.ConsumeValueOrDie());
+				FilePackageStore->Mount(Pak->IoContainerHeader.Get(), PakOrder);
 			}
 			else
 			{
 				bIoStoreSuccess = false;
-				UE_LOG(LogPakFile, Warning, TEXT("Failed to mount IoStore container \"%s\" [%s]"), *ContainerPath, *IoStatus.Status().ToString());
+				UE_LOG(LogPakFile, Warning, TEXT("Failed to mount IoStore container \"%s\" [%s]"), *ContainerPath, *MountResult.Status().ToString());
 			}
 		}
 
@@ -7666,9 +7675,12 @@ bool FPakPlatformFile::Unmount(const TCHAR* InPakFilename)
 
 		if (IoDispatcherFileBackend.IsValid())
 		{
+			if (UnmountedPak)
+			{
+				FilePackageStore->Unmount(UnmountedPak->IoContainerHeader.Get());
+			}
 			FString ContainerPath = FPaths::ChangeExtension(InPakFilename, FString());
-			FIoStatus Status = FIoDispatcher::Get().Unmount(*ContainerPath);
-			bRemovedContainerFile = Status.IsOk();
+			bRemovedContainerFile = IoDispatcherFileBackend->Unmount(*ContainerPath);
 		}
 
 		if (UnmountedPak)
diff --git a/Engine/Source/Runtime/PakFile/Public/IPlatformFilePak.h b/Engine/Source/Runtime/PakFile/Public/IPlatformFilePak.h
--- a/Engine/Source/Runtime/PakFile/Public/IPlatformFilePak.h
+++ b/Engine/Source/Runtime/PakFile/Public/IPlatformFilePak.h
@@ -16,11 +16,14 @@
 #include "GenericPlatform/GenericPlatformChunkInstall.h"
 #include "Serialization/MemoryImage.h"
 #include "Templates/RefCounting.h"
+#include "IO/IoContainerId.h"
 #include "Containers/Ticker.h"
 
 class FChunkCacheWorker;
 class IAsyncReadFileHandle;
 struct IIoDispatcherFileBackend;
+class FFilePackageStore;
+struct FIoContainerHeader;
 
 PAKFILE_API DECLARE_LOG_CATEGORY_EXTERN(LogPakFile, Log, All);
 DECLARE_FLOAT_ACCUMULATOR_STAT_EXTERN(TEXT("Total pak file read time"), STAT_PakFile_Read, STATGROUP_PakFile, PAKFILE_API);
@@ -793,6 +796,8 @@ private:
 	/** Record of whether the pak file is still mounted, so PakPrecacher can reject requests to register it. */
 	bool bIsMounted;
 
+	TUniquePtr<FIoContainerHeader> IoContainerHeader;
+
 	static inline int32 CDECL CompareFilenameHashes(const void* Left, const void* Right)
 	{
 		const uint64* LeftHash = (const uint64*)Left;
@@ -2027,6 +2032,7 @@ class PAKFILE_API FPakPlatformFile : public IPlatformFile
 	/** The filename for the gameusersettings ini file, used for excluding ini files, but not gameusersettings */
 	FString GameUserSettingsIniFilename;
 	TSharedPtr<IIoDispatcherFileBackend> IoDispatcherFileBackend;
+	TSharedPtr<FFilePackageStore> FilePackageStore;
 
 	FTSTicker::FDelegateHandle RetireReadersHandle;
 
diff --git a/Engine/Source/Runtime/StorageServerClient/Private/StorageServerIoDispatcherBackend.cpp b/Engine/Source/Runtime/StorageServerClient/Private/StorageServerIoDispatcherBackend.cpp
--- a/Engine/Source/Runtime/StorageServerClient/Private/StorageServerIoDispatcherBackend.cpp
+++ b/Engine/Source/Runtime/StorageServerClient/Private/StorageServerIoDispatcherBackend.cpp
@@ -110,11 +110,6 @@ uint32 FStorageServerIoDispatcherBackend::Run()
 	return 0;
 }
 
-void FStorageServerIoDispatcherBackend::AppendMountedContainers(TSet<FIoContainerId>& Containers)
-{
-	Containers.Add(FIoContainerId::FromName(TEXT("global")));
-}
-
 bool FStorageServerIoDispatcherBackend::Resolve(FIoRequestImpl* Request)
 {
 	check(Request);
diff --git a/Engine/Source/Runtime/StorageServerClient/Private/StorageServerIoDispatcherBackend.h b/Engine/Source/Runtime/StorageServerClient/Private/StorageServerIoDispatcherBackend.h
--- a/Engine/Source/Runtime/StorageServerClient/Private/StorageServerIoDispatcherBackend.h
+++ b/Engine/Source/Runtime/StorageServerClient/Private/StorageServerIoDispatcherBackend.h
@@ -18,7 +18,6 @@ public:
 	~FStorageServerIoDispatcherBackend();
 
 	void Initialize(TSharedRef<const FIoDispatcherBackendContext> Context) override;
-	void AppendMountedContainers(TSet<FIoContainerId>& Containers) override;
 	bool Resolve(FIoRequestImpl* Request) override;
 	void CancelIoRequest(FIoRequestImpl* Request) override {};
 	void UpdatePriorityForIoRequest(FIoRequestImpl* Request) override {};
diff --git a/Engine/Source/Runtime/StorageServerClient/Private/StorageServerPackageStore.cpp b/Engine/Source/Runtime/StorageServerClient/Private/StorageServerPackageStore.cpp
new file mode 100644
--- /dev/null
+++ b/Engine/Source/Runtime/StorageServerClient/Private/StorageServerPackageStore.cpp
@@ -0,0 +1,63 @@
+// Copyright Epic Games, Inc. All Rights Reserved.
+
+#include "StorageServerPackageStore.h"
+#include "IO/IoContainerHeader.h"
+#include "IO/IoDispatcher.h"
+#include "StorageServerConnection.h"
+
+#if !UE_BUILD_SHIPPING
+
+FStorageServerPackageStore::FStorageServerPackageStore(FStorageServerConnection& Connection)
+{
+	FIoChunkId HeaderChunkId = CreateIoChunkId(FIoContainerId::FromName(TEXT("global")).Value(), 0, EIoChunkType::ContainerHeader);
+	Connection.ReadChunkRequest(HeaderChunkId, 0, uint64(-1), [this](FStorageServerResponse& ResponseStream)
+	{
+		if (ResponseStream.IsOk())
+		{
+			FIoContainerHeader ContainerHeader;
+			ResponseStream << ContainerHeader;
+			StoreEntriesData = MoveTemp(ContainerHeader.StoreEntries);
+			StoreEntriesMap.Reserve(ContainerHeader.PackageCount);
+			TArrayView<const FFilePackageStoreEntry> StoreEntries(reinterpret_cast<const FFilePackageStoreEntry*>(StoreEntriesData.GetData()), ContainerHeader.PackageCount);
+			int32 Index = 0;
+			for (const FFilePackageStoreEntry& StoreEntry : StoreEntries)
+			{
+				const FPackageId& PackageId = ContainerHeader.PackageIds[Index];
+				check(PackageId.IsValid());
+				StoreEntriesMap.FindOrAdd(PackageId, &StoreEntry);
+				++Index;
+			}
+		}
+	});
+}
+
+bool FStorageServerPackageStore::DoesPackageExist(FPackageId PackageId)
+{
+	return PackageId.IsValid() && StoreEntriesMap.Contains(PackageId);
+}
+
+FPackageStoreEntryHandle FStorageServerPackageStore::GetPackageEntryHandle(FPackageId PackageId, const FName& PackageName)
+{
+	const FFilePackageStoreEntry* FindEntry = StoreEntriesMap.FindRef(PackageId);
+	const uint64 Handle = reinterpret_cast<uint64>(FindEntry);
+	return FPackageStoreEntryHandle::Create(Handle, Handle ? EPackageStoreEntryStatus::Ok : EPackageStoreEntryStatus::Missing);
+}
+
+FPackageStoreEntry FStorageServerPackageStore::GetPackageEntry(FPackageStoreEntryHandle Handle)
+{
+	check(Handle.IsValid());
+	const FFilePackageStoreEntry* Entry = reinterpret_cast<const FFilePackageStoreEntry*>(Handle.Value());
+	check(Entry);
+	return FPackageStoreEntry
+	{
+		FPackageStoreExportInfo
+		{
+			Entry->ExportCount,
+			Entry->ExportBundleCount
+		},
+		MakeArrayView(Entry->ImportedPackages.Data(), Entry->ImportedPackages.Num()),
+		MakeArrayView(Entry->ShaderMapHashes.Data(), Entry->ShaderMapHashes.Num())
+	};
+}
+
+#endif
\ No newline at end of file
diff --git a/Engine/Source/Runtime/StorageServerClient/Private/StorageServerPackageStore.h b/Engine/Source/Runtime/StorageServerClient/Private/StorageServerPackageStore.h
new file mode 100644
--- /dev/null
+++ b/Engine/Source/Runtime/StorageServerClient/Private/StorageServerPackageStore.h
@@ -0,0 +1,45 @@
+// Copyright Epic Games, Inc. All Rights Reserved.
+
+#pragma once
+
+#include "IO/PackageStore.h"
+
+#if !UE_BUILD_SHIPPING
+
+class FStorageServerConnection;
+struct FFilePackageStoreEntry;
+
+class FStorageServerPackageStore
+	: public FPackageStoreBase
+{
+public:
+	FStorageServerPackageStore(FStorageServerConnection& Connection);
+	virtual ~FStorageServerPackageStore() = default;
+
+	virtual void Initialize() override
+	{
+	}
+
+	virtual void Lock() override
+	{
+	}
+
+	virtual void Unlock() override
+	{
+	}
+
+	virtual bool DoesPackageExist(FPackageId PackageId) override;
+	virtual FPackageStoreEntryHandle GetPackageEntryHandle(FPackageId PackageId, const FName& PackageName) override;
+	virtual FPackageStoreEntry GetPackageEntry(FPackageStoreEntryHandle Handle) override;
+	
+	virtual bool GetPackageRedirectInfo(FPackageId PackageId, FName& OutSourcePackageName, FPackageId& OutRedirectedToPackageId) override
+	{
+		return false;
+	}
+
+private:
+	TArray<uint8> StoreEntriesData;
+	TMap<FPackageId, const FFilePackageStoreEntry*> StoreEntriesMap;
+};
+
+#endif
\ No newline at end of file
diff --git a/Engine/Source/Runtime/StorageServerClient/Private/StorageServerPlatformFile.cpp b/Engine/Source/Runtime/StorageServerClient/Private/StorageServerPlatformFile.cpp
--- a/Engine/Source/Runtime/StorageServerClient/Private/StorageServerPlatformFile.cpp
+++ b/Engine/Source/Runtime/StorageServerClient/Private/StorageServerPlatformFile.cpp
@@ -10,6 +10,8 @@
 #include "Modules/ModuleManager.h"
 #include "Misc/StringBuilder.h"
 #include "Algo/Replace.h"
+#include "StorageServerPackageStore.h"
+#include "CookOnTheFlyPackageStore.h"
 
 #if WITH_COTF
 #include "Modules/ModuleManager.h"
@@ -281,14 +283,23 @@ bool FStorageServerPlatformFile::Initialize(IPlatformFile* Inner, const TCHAR* C
 				FIoDispatcher& IoDispatcher = FIoDispatcher::Get();
 				TSharedRef<FStorageServerIoDispatcherBackend> IoDispatcherBackend = MakeShared<FStorageServerIoDispatcherBackend>(*Connection.Get());
 				IoDispatcher.Mount(IoDispatcherBackend);
-#if WITH_COTF
-				if (IsRunningCookOnTheFly())
+				bResult = true;
+				FCoreDelegates::CreatePackageStore.BindLambda([this]() -> TSharedPtr<IPackageStore>
 				{
-					UE::Cook::ICookOnTheFlyModule& CookOnTheFlyModule = FModuleManager::LoadModuleChecked<UE::Cook::ICookOnTheFlyModule>(TEXT("CookOnTheFly"));
-					CookOnTheFlyModule.GetServerConnection().OnMessage().AddRaw(this, &FStorageServerPlatformFile::OnCookOnTheFlyMessage);
-				}
+#if WITH_COTF
+					if (IsRunningCookOnTheFly())
+					{
+						UE::Cook::ICookOnTheFlyModule& CookOnTheFlyModule = FModuleManager::LoadModuleChecked<UE::Cook::ICookOnTheFlyModule>(TEXT("CookOnTheFly"));
+						UE::Cook::ICookOnTheFlyServerConnection& CotfConnection = CookOnTheFlyModule.GetServerConnection();
+						CotfConnection.OnMessage().AddRaw(this, &FStorageServerPlatformFile::OnCookOnTheFlyMessage);
+						return MakeCookOnTheFlyPackageStore(CotfConnection);
+					}
+					else
 #endif
-				bResult = true;
+					{
+						return MakeShared<FStorageServerPackageStore>(*Connection.Get());
+					}
+				});
 			}
 			else
 			{
diff --git a/Engine/Source/Runtime/StorageServerClient/Private/StorageServerPlatformFile.h b/Engine/Source/Runtime/StorageServerClient/Private/StorageServerPlatformFile.h
--- a/Engine/Source/Runtime/StorageServerClient/Private/StorageServerPlatformFile.h
+++ b/Engine/Source/Runtime/StorageServerClient/Private/StorageServerPlatformFile.h
@@ -11,6 +11,7 @@
 
 class FStorageServerFileHandle;
 class FStorageServerConnection;
+class IPackageStore;
 
 #if WITH_COTF
 namespace UE::Cook
