Made Zen COTF API more private
#rnx

diff --git a/Engine/Source/Editor/UnrealEd/Private/Cooker/IoStoreCookOnTheFlyRequestManager.cpp b/Engine/Source/Editor/UnrealEd/Private/Cooker/IoStoreCookOnTheFlyRequestManager.cpp
--- a/Engine/Source/Editor/UnrealEd/Private/Cooker/IoStoreCookOnTheFlyRequestManager.cpp
+++ b/Engine/Source/Editor/UnrealEd/Private/Cooker/IoStoreCookOnTheFlyRequestManager.cpp
@@ -3,8 +3,8 @@
 #include "IoStoreCookOnTheFlyRequestManager.h"
 #include "CookOnTheFlyServerInterface.h"
 #include "CookOnTheFly.h"
+#include "CookOnTheFlyMessages.h"
 #include "Modules/ModuleManager.h"
-#include "CookOnTheFlyPackageStore.h"
 #include "PackageStoreWriter.h"
 #include "ShaderCompiler.h"
 #include "HAL/PlatformTime.h"
@@ -345,11 +345,11 @@ private:
 	bool HandleGetCookedPackagesRequest(UE::Cook::FCookOnTheFlyClient Client, const UE::Cook::FCookOnTheFlyRequest& Request, UE::Cook::FCookOnTheFlyResponse& Response)
 	{
 		using namespace UE::Cook;
-		using namespace UE::PackageStore::Messaging;
+		using namespace UE::ZenCookOnTheFly::Messaging;
 
 		FPlatformContext& Context = GetContext(Client.PlatformName);
 		
-		UE::PackageStore::Messaging::FPackageStoreData PackageStoreData;
+		UE::ZenCookOnTheFly::Messaging::FPackageStoreData PackageStoreData;
 		
 		Context.PackageWriter->GetEntries(
 			[&Context, &PackageStoreData](TArrayView<const FPackageStoreEntryResource> Entries,
@@ -383,7 +383,7 @@ private:
 
 	bool HandleCookPackageRequest(UE::Cook::FCookOnTheFlyClient Client, const UE::Cook::FCookOnTheFlyRequest& Request, UE::Cook::FCookOnTheFlyResponse& Response)
 	{
-		using namespace UE::PackageStore::Messaging;
+		using namespace UE::ZenCookOnTheFly::Messaging;
 
 		TRACE_CPUPROFILER_EVENT_SCOPE(CookOnTheFly::HandleCookPackageRequest);
 
@@ -434,7 +434,7 @@ private:
 	void OnPackageCooked(const IPackageStoreWriter::FCommitEventArgs& EventArgs)
 	{
 		using namespace UE::Cook;
-		using namespace UE::PackageStore::Messaging;
+		using namespace UE::ZenCookOnTheFly::Messaging;
 
 		TRACE_CPUPROFILER_EVENT_SCOPE(CookOnTheFly::OnPackageCooked);
 
@@ -496,7 +496,7 @@ private:
 		FCookOnTheFlyPackageTracker::FCompletedPackages&& NewCompletedPackages, TArrayView<const FPackageStoreEntryResource> Entries)
 	{
 		using namespace UE::Cook;
-		using namespace UE::PackageStore::Messaging;
+		using namespace UE::ZenCookOnTheFly::Messaging;
 		if (NewCompletedPackages.CookedEntryIndices.Num() || NewCompletedPackages.FailedPackages.Num())
 		{
 			TRACE_CPUPROFILER_EVENT_SCOPE(CookOnTheFly::SendCookedPackagesMessage);
diff --git a/Engine/Source/Runtime/CookOnTheFly/Internal/CookOnTheFly.h b/Engine/Source/Runtime/CookOnTheFly/Internal/CookOnTheFly.h
new file mode 100644
--- /dev/null
+++ b/Engine/Source/Runtime/CookOnTheFly/Internal/CookOnTheFly.h
@@ -0,0 +1,382 @@
+// Copyright Epic Games, Inc. All Rights Reserved.
+
+#pragma once
+
+#include "CoreMinimal.h"
+#include "Modules/ModuleInterface.h"
+#include "Async/Future.h"
+#include "Misc/Timespan.h"
+#include "Serialization/MemoryReader.h"
+#include "Serialization/MemoryWriter.h"
+
+COOKONTHEFLY_API DECLARE_LOG_CATEGORY_EXTERN(LogCookOnTheFly, Log, All);
+
+namespace UE { namespace Cook
+{
+
+enum
+{ 
+	/* The default port used by the the cook-on-the-fly server. */
+	DefaultCookOnTheFlyServingPort = 42899
+};
+
+/**
+ * Flags and message types to be used with the cook-on-the-fly server.
+ *
+ * The first 8 bits specifies the message type, i.e request, response or a on-way message.
+ */
+enum class ECookOnTheFlyMessage : uint32
+{
+	/* Represents no message. */
+	None				= 0x00,
+	
+	/* A one way message. */
+	Message				= 0x01,
+	/* A request message. */
+	Request				= 0x02,
+	/* A response message. */
+	Response			= 0x04,
+	TypeFlags			= 0x0F,
+
+	/* The handshake request message. */
+	Handshake			= 0x10,
+	/* Request to cook a package. */
+	CookPackage			= 0x20,
+	/* Get all currenlty cooked packages. */
+	GetCookedPackages	= 0x30,
+	/* Request to recompile shaders. */
+	RecompileShaders	= 0x40,
+	/* One way message indicating that one or more packages has been cooked. */
+	PackagesCooked		= 0x50,
+	/* One way message indicating that one or more files has been added. */
+	FilesAdded			= 0x60,
+	/* Heartbeat message. */
+	Heartbeat			= 0x70,
+};
+ENUM_CLASS_FLAGS(ECookOnTheFlyMessage);
+
+/**
+ * Returns a string from the specified message.
+ */
+inline const TCHAR* LexToString(ECookOnTheFlyMessage Message)
+{
+	EnumRemoveFlags(Message, ECookOnTheFlyMessage::TypeFlags);
+
+	switch (Message)
+	{
+		case ECookOnTheFlyMessage::None:
+			return TEXT("None");
+		case ECookOnTheFlyMessage::Handshake:
+			return TEXT("Handshake");
+		case ECookOnTheFlyMessage::CookPackage:
+			return TEXT("CookPackage");
+		case ECookOnTheFlyMessage::GetCookedPackages:
+			return TEXT("GetCookedPackages");
+		case ECookOnTheFlyMessage::RecompileShaders:
+			return TEXT("RecompileShaders");
+		case ECookOnTheFlyMessage::PackagesCooked:
+			return TEXT("PackagesCooked");
+		case ECookOnTheFlyMessage::FilesAdded:
+			return TEXT("FilesAdded");
+		case ECookOnTheFlyMessage::Heartbeat:
+			return TEXT("Heartbeat");
+		default:
+			return TEXT("Unknown");
+	};
+}
+
+/**
+ * Cook-on-the-fly message status.
+ */
+enum class ECookOnTheFlyMessageStatus : uint32
+{
+	/** No status. */
+	None,
+	/** The message is successful. */
+	Ok,
+	/** The message failed. */
+	Error
+};
+
+/**
+ * Returns a string from the specified message status.
+ */
+inline const TCHAR* LexToString(ECookOnTheFlyMessageStatus Status)
+{
+	switch (Status)
+	{
+		case ECookOnTheFlyMessageStatus::None:
+			return TEXT("None");
+		case ECookOnTheFlyMessageStatus::Ok:
+			return TEXT("Ok");
+		case ECookOnTheFlyMessageStatus::Error:
+			return TEXT("Error");
+		default:
+			return TEXT("Unknown");
+	}
+}
+
+/**
+ * Cook-on-the-fly message header.
+ */
+struct FCookOnTheFlyMessageHeader
+{
+	/** Type of message */
+	ECookOnTheFlyMessage MessageType = ECookOnTheFlyMessage::None;
+	/** The message status. */
+	ECookOnTheFlyMessageStatus MessageStatus = ECookOnTheFlyMessageStatus::None;
+	/** Sender id */
+	uint32 SenderId = 0;
+	/** Correlation id, used to match response with request. */
+	uint32 CorrelationId = 0;
+	/** When the message was sent. */
+	int64 Timestamp = 0;
+
+	COOKONTHEFLY_API FString ToString() const;
+
+	COOKONTHEFLY_API friend FArchive& operator<<(FArchive& Ar, FCookOnTheFlyMessageHeader& Header);
+};
+
+/**
+ * Cook-on-the-fly message.
+ */
+class FCookOnTheFlyMessage
+{
+public:
+	/** Creates a new instance of a cook-on-the-fly message. */
+	COOKONTHEFLY_API FCookOnTheFlyMessage() = default;
+
+	/** Creates a new instance of a cook-on-the-fly message with the specified message type. */
+	explicit FCookOnTheFlyMessage(ECookOnTheFlyMessage MessageType)
+	{
+		Header.MessageType = MessageType;
+	}
+
+	/** Returns the message header. */
+	FCookOnTheFlyMessageHeader& GetHeader()
+	{
+		return Header;
+	}
+
+	/** Returns the message header. */
+	const FCookOnTheFlyMessageHeader& GetHeader() const
+	{
+		return Header;
+	}
+
+	/** Set a new message header. */
+	void SetHeader(const FCookOnTheFlyMessageHeader& InHeader)
+	{
+		Header = InHeader;
+	}
+
+	/** Sets the message status. */
+	void SetStatus(ECookOnTheFlyMessageStatus InStatus)
+	{
+		Header.MessageStatus = InStatus;
+	}
+
+	/** Returns the message status. */
+	ECookOnTheFlyMessageStatus GetStatus() const
+	{
+		return Header.MessageStatus;
+	}
+
+	/** Returns whether the message stauts is OK. */
+	bool IsOk() const
+	{
+		return Header.MessageStatus == ECookOnTheFlyMessageStatus::Ok;
+	}
+
+	/** Set the message body. */
+	COOKONTHEFLY_API void SetBody(TArray<uint8> InBody);
+
+	/** Set body to serializable type. */
+	template<typename BodyType>
+	void SetBodyTo(BodyType InBody)
+	{
+		Body.Empty();
+		FMemoryWriter Ar(Body);
+		Ar << InBody;
+	}
+	
+	/** Returns the message body. */
+	TArray<uint8>& GetBody()
+	{
+		return Body;
+	}
+
+	/** Returns the message body. */
+	const TArray<uint8>& GetBody() const
+	{
+		return Body;
+	}
+
+	/** Serialize the body as the specified type. */
+	template<typename BodyType>
+	BodyType GetBodyAs() const
+	{
+		BodyType Type;
+		FMemoryReader Ar(Body);
+		Ar << Type;
+
+		return MoveTemp(Type);
+	}
+
+	/** Returns the total size of the message header and message body. */
+	int64 TotalSize() const
+	{
+		return sizeof(FCookOnTheFlyMessageHeader) + Body.Num();
+	}
+
+	/** Creates an archive for ready the message body. */
+	COOKONTHEFLY_API TUniquePtr<FArchive> ReadBody() const;
+
+	/** Creates an archive for writing the message body. */
+	COOKONTHEFLY_API TUniquePtr<FArchive> WriteBody();
+
+	COOKONTHEFLY_API friend FArchive& operator<<(FArchive& Ar, FCookOnTheFlyMessage& Message);
+
+protected:
+	FCookOnTheFlyMessageHeader Header;
+	TArray<uint8> Body;
+};
+
+using FCookOnTheFlyRequest = FCookOnTheFlyMessage;
+
+using FCookOnTheFlyResponse = FCookOnTheFlyMessage;
+
+/**
+ * Connection status
+ */
+enum class ECookOnTheFlyConnectionStatus
+{
+	Disconnected,
+	Connected
+};
+
+/**
+ * A connected cook-on-the-fly client.
+ */
+struct FCookOnTheFlyClient
+{
+	/** A client ID set by the server. */
+	uint32 ClientId = 0;
+	/** The platform. */
+	FName PlatformName;
+};
+
+using FCookOnTheFlyRequestHandler = TFunction<bool(FCookOnTheFlyClient, const FCookOnTheFlyRequest&, FCookOnTheFlyResponse&)>;
+
+using FCookOnTheFlyClientConnectionHandler = TFunction<bool(FCookOnTheFlyClient, ECookOnTheFlyConnectionStatus)>;
+
+using FFillRequest = TFunction<void(FArchive&)>;
+
+using FProcessResponse = TFunction<bool(FArchive&)>;
+
+/**
+ * Cook-on-the-fly connection server options.
+ */
+struct FCookOnTheFlyServerOptions
+{
+	/** The port to listen for new connections. */
+	int32 Port = DefaultCookOnTheFlyServingPort;
+
+	/** Callback invoked when a client has connected or disconnected. */
+	FCookOnTheFlyClientConnectionHandler HandleClientConnection;
+
+	/** Callback invoked when the server receives a new request. */
+	FCookOnTheFlyRequestHandler HandleRequest;
+};
+
+/**
+ * Cook-on-the-fly host address.
+ */
+struct FCookOnTheFlyHostOptions
+{
+	/** Host address. */
+	TArray<FString> Hosts;
+	/** Host port. */
+	int32 Port = DefaultCookOnTheFlyServingPort;
+	/** How long to wait for the server to start. */
+	FTimespan ServerStartupWaitTime;
+};
+
+/**
+ * A connection server used to communicate with cook-on-the-fly clients.
+ */
+class ICookOnTheFlyConnectionServer
+{
+public:
+	virtual ~ICookOnTheFlyConnectionServer() { }
+
+	/** Start the cook-on-the-fly server. */
+	virtual bool StartServer() = 0;
+
+	/** Stop the cook-on-the-fly server. */
+	virtual void StopServer() = 0;
+
+	/** Broadcast message to all connected clients for the specified platform. */
+	virtual bool BroadcastMessage(const FCookOnTheFlyMessage& Message, const FName& PlatformName = NAME_None) = 0;
+};
+
+/**
+ * A connection used to communicate with the cook-on-the-fly server.
+ */
+class ICookOnTheFlyServerConnection
+{
+public:
+	virtual ~ICookOnTheFlyServerConnection() { }
+
+	/**
+	 * Returns whether connected to the cook-on-the-fly server.
+	 */
+	virtual bool IsConnected() const = 0;
+	
+	/**
+	 * Disconnect from the server.
+	 */
+	virtual void Disconnect() = 0;
+
+	/**
+	 * Sends a request to the server.
+	 *
+	 * @param Request The request message to send.
+	 * @param FillRequest Callback to populate the request message..
+	 * @param ProcessResponse Callback to process the response message.
+	 */
+	//virtual bool SendRequest(ECookOnTheFlyMessage Request, FFillRequest&& FillRequest, FProcessResponse&& ProcessResponse) = 0;
+	virtual TFuture<FCookOnTheFlyResponse> SendRequest(const FCookOnTheFlyRequest& Request) = 0;
+
+	/**
+	 * Event triggered when a new message has been sent from the server.
+	 */
+	DECLARE_EVENT_OneParam(ICookOnTheFlyServerConnection, FMessageEvent, const FCookOnTheFlyMessage&);
+	virtual FMessageEvent& OnMessage() = 0;
+};
+
+/**
+ * Cook-on-the-fly module
+ */
+class ICookOnTheFlyModule
+	: public IModuleInterface
+{
+public:
+	virtual ~ICookOnTheFlyModule() { }
+
+	/**
+	 * Creates a new instance of a cook-on-the-fly connection server.
+	 *
+	 * @param Options The cook-on-the-fly connection server options.
+	 */
+	virtual TUniquePtr<ICookOnTheFlyConnectionServer> CreateConnectionServer(FCookOnTheFlyServerOptions Options) = 0;
+
+	/**
+	 * Connect to the cook-on-the-fly server.
+	 *
+	 * @param HostOptions Cook-on-the-fly host options.
+	 */
+	virtual TUniquePtr<ICookOnTheFlyServerConnection> ConnectToServer(const FCookOnTheFlyHostOptions& HostOptions) = 0;
+};
+
+}} // namespace UE::Cook
diff --git a/Engine/Source/Runtime/CookOnTheFly/Internal/CookOnTheFlyMessages.h b/Engine/Source/Runtime/CookOnTheFly/Internal/CookOnTheFlyMessages.h
new file mode 100644
--- /dev/null
+++ b/Engine/Source/Runtime/CookOnTheFly/Internal/CookOnTheFlyMessages.h
@@ -0,0 +1,60 @@
+// Copyright Epic Games, Inc. All Rights Reserved.
+
+#pragma once
+
+#include "CoreMinimal.h"
+
+#if WITH_COTF
+
+#include "IO/PackageStore.h"
+
+namespace UE { namespace ZenCookOnTheFly { namespace Messaging
+{
+
+struct FPackageStoreData
+{
+	TArray<FPackageStoreEntryResource> CookedPackages;
+	TArray<FPackageId> FailedPackages;
+	int32 TotalCookedPackages = 0;
+	int32 TotalFailedPackages = 0;
+
+	COOKONTHEFLY_API friend FArchive& operator<<(FArchive& Ar, FPackageStoreData& PackageStoreData);
+};
+
+struct FCookPackageRequest
+{
+	FPackageId PackageId;
+	
+	COOKONTHEFLY_API friend FArchive& operator<<(FArchive& Ar, FCookPackageRequest& Request);
+};
+
+struct FCookPackageResponse
+{
+	EPackageStoreEntryStatus Status = EPackageStoreEntryStatus::None;
+
+	COOKONTHEFLY_API friend FArchive& operator<<(FArchive& Ar, FCookPackageResponse& Response);
+};
+
+struct FPackagesCookedMessage
+{
+	FPackageStoreData PackageStoreData;
+
+	friend FArchive& operator<<(FArchive& Ar, FPackagesCookedMessage& PackagesCookedMessage)
+	{
+		return Ar << PackagesCookedMessage.PackageStoreData;
+	}
+};
+
+struct FGetCookedPackagesResponse
+{
+	FPackageStoreData PackageStoreData;
+	
+	friend FArchive& operator<<(FArchive& Ar, FGetCookedPackagesResponse& GetCookedPackagesResponse)
+	{
+		return Ar << GetCookedPackagesResponse.PackageStoreData;
+	}
+};
+
+}}} // namesapce UE::ZenCookOnTheFly::Messaging
+
+#endif // WITH_COTF
diff --git a/Engine/Source/Runtime/CookOnTheFly/Private/CookOnTheFly.cpp b/Engine/Source/Runtime/CookOnTheFly/Private/CookOnTheFly.cpp
--- a/Engine/Source/Runtime/CookOnTheFly/Private/CookOnTheFly.cpp
+++ b/Engine/Source/Runtime/CookOnTheFly/Private/CookOnTheFly.cpp
@@ -61,91 +61,4 @@ TUniquePtr<FArchive> FCookOnTheFlyMessage::WriteBody()
 	return MakeUnique<FMemoryWriter>(Body);
 }
 
-bool GetCookOnTheFlyHost(UE::Cook::FCookOnTheFlyHostOptions& OutHostOptions)
-{
-	bool bOptionsOk = false;
-
-	if (FParse::Param(FCommandLine::Get(), TEXT("CookOnTheFly")))
-	{
-		// Cook-on-the-fly expects the same host as the Zen storage server
-		FString Host;
-		if (FParse::Value(FCommandLine::Get(), TEXT("-ZenStoreHost="), Host))
-		{
-			if (!Host.ParseIntoArray(OutHostOptions.Hosts, TEXT("+"), true))
-			{
-				OutHostOptions.Hosts.Add(Host);
-			}
-
-			bOptionsOk = true;
-		}
-	}
-
-	double ServerWaitTimeInSeconds;
-	if (FParse::Value(FCommandLine::Get(), TEXT("-CookOnTheFlyServerWaitTime="), ServerWaitTimeInSeconds))
-	{
-		OutHostOptions.ServerStartupWaitTime = FTimespan::FromSeconds(ServerWaitTimeInSeconds);
-	}
-
-	return bOptionsOk;
-}
-
-bool SendCookOnTheFlyRequest(ECookOnTheFlyMessage RequestType, TFunction<void(FArchive&)>&& FillRequest, TFunction<bool(FArchive&)>&& ProcessResponse)
-{
-	if (!IsRunningCookOnTheFly())
-	{
-		return false;
-	}
-
-	static FString FileHostIP;
-	static bool bUseFileServer = FParse::Value(FCommandLine::Get(), TEXT("filehostip"), FileHostIP);
-
-	if (bUseFileServer)
-	{
-		struct FFileServerMessageHandler
-			: public IPlatformFile::IFileServerMessageHandler
-		{
-			FFileServerMessageHandler(TFunction<void(FArchive&)>&& InFillRequest, TFunction<bool(FArchive&)>&& InProcessResponse)
-				: FillRequestHandler(InFillRequest)
-				, ProcessResponseHandler(InProcessResponse)
-			{ }
-
-			virtual void FillPayload(FArchive& Payload) override
-			{
-				FillRequestHandler(Payload);
-			}
-
-			virtual void ProcessResponse(FArchive& Response) override
-			{
-				ProcessResponseHandler(Response);
-			}
-
-			TFunction<void(FArchive&)> FillRequestHandler;
-			TFunction<bool(FArchive&)> ProcessResponseHandler;
-		};
-
-		FFileServerMessageHandler MessageHandler(MoveTemp(FillRequest), MoveTemp(ProcessResponse));
-		return IFileManager::Get().SendMessageToServer(LexToString(RequestType), &MessageHandler);
-	}
-	else
-	{
-		static ICookOnTheFlyModule& CookOnTheFlyModule = FModuleManager::LoadModuleChecked<ICookOnTheFlyModule>(TEXT("CookOnTheFly"));
-		ICookOnTheFlyServerConnection& ServerConnection = CookOnTheFlyModule.GetServerConnection();
-
-		FCookOnTheFlyRequest Request(RequestType);
-		{
-			TUniquePtr<FArchive> Ar = Request.WriteBody();
-			FillRequest(*Ar);
-		}
-
-		FCookOnTheFlyResponse Response = ServerConnection.SendRequest(Request).Get();
-		if (Response.IsOk())
-		{
-			TUniquePtr<FArchive> Ar = Response.ReadBody();
-			ProcessResponse(*Ar);
-		}
-
-		return Response.IsOk();
-	}
-}
-
 }} // namesapce UE::Cook
diff --git a/Engine/Source/Runtime/CookOnTheFly/Private/CookOnTheFlyMessages.cpp b/Engine/Source/Runtime/CookOnTheFly/Private/CookOnTheFlyMessages.cpp
new file mode 100644
--- /dev/null
+++ b/Engine/Source/Runtime/CookOnTheFly/Private/CookOnTheFlyMessages.cpp
@@ -0,0 +1,42 @@
+// Copyright Epic Games, Inc. All Rights Reserved.
+
+#include "CookOnTheFlyMessages.h"
+
+#if WITH_COTF
+
+namespace UE { namespace ZenCookOnTheFly { namespace Messaging
+{
+
+FArchive& operator<<(FArchive& Ar, FPackageStoreData& PackageStoreData)
+{
+	Ar << PackageStoreData.CookedPackages;
+	Ar << PackageStoreData.FailedPackages;
+	Ar << PackageStoreData.TotalCookedPackages;
+	Ar << PackageStoreData.TotalFailedPackages;
+
+	return Ar;
+}
+
+FArchive& operator<<(FArchive& Ar, FCookPackageRequest& Request)
+{
+	Ar << Request.PackageId;
+
+	return Ar;
+}
+
+FArchive& operator<<(FArchive& Ar, FCookPackageResponse& Response)
+{
+	uint32 Status = static_cast<uint32>(Response.Status);
+	Ar << Status;
+
+	if (Ar.IsLoading())
+	{
+		Response.Status = static_cast<EPackageStoreEntryStatus>(Status);
+	}
+
+	return Ar;
+}
+
+}}} // namesapce UE::ZenCookOnTheFly::Messaging
+
+#endif // WITH_COTF
diff --git a/Engine/Source/Runtime/CookOnTheFly/Private/CookOnTheFlyModule.cpp b/Engine/Source/Runtime/CookOnTheFly/Private/CookOnTheFlyModule.cpp
--- a/Engine/Source/Runtime/CookOnTheFly/Private/CookOnTheFlyModule.cpp
+++ b/Engine/Source/Runtime/CookOnTheFly/Private/CookOnTheFlyModule.cpp
@@ -10,7 +10,7 @@ namespace UE { namespace Cook
 
 TUniquePtr<ICookOnTheFlyConnectionServer> MakeCookOnTheFlyConnectionServer(FCookOnTheFlyServerOptions Options);
 	
-TUniquePtr<ICookOnTheFlyServerConnection> MakeServerConnection();
+TUniquePtr<ICookOnTheFlyServerConnection> MakeServerConnection(const UE::Cook::FCookOnTheFlyHostOptions& HostOptions);
 
 }}
 
@@ -25,15 +25,8 @@ public:
 
 	virtual void StartupModule( ) override { }
 
-	virtual void ShutdownModule( ) override
-	{ 
-		if (ServerConnection.IsValid())
-		{
-			ServerConnection->Disconnect();
-			ServerConnection.Reset();
-		}
-	}
-
+	virtual void ShutdownModule( ) override { }
+	
 	virtual bool SupportsDynamicReloading( ) override { return false; }
 
 	//
@@ -45,25 +38,10 @@ public:
 		return UE::Cook::MakeCookOnTheFlyConnectionServer(Options);
 	}
 
-	virtual bool ConnectToServer(const UE::Cook::FCookOnTheFlyHostOptions& HostOptions) override
-	{
-		if (!ServerConnection.IsValid())
-		{
-			ServerConnection = UE::Cook::MakeServerConnection();
-			return ServerConnection->Connect(HostOptions);
-		}
-
-		return ServerConnection->IsConnected();
-	}
-
-	virtual UE::Cook::ICookOnTheFlyServerConnection& GetServerConnection() override
+	virtual TUniquePtr<UE::Cook::ICookOnTheFlyServerConnection> ConnectToServer(const UE::Cook::FCookOnTheFlyHostOptions& HostOptions) override
 	{
-		checkf(ServerConnection.IsValid(), TEXT("Call connect before accessing the server connection"));
-		return *ServerConnection;
+		return UE::Cook::MakeServerConnection(HostOptions);
 	}
-
-private:
-	TUniquePtr<UE::Cook::ICookOnTheFlyServerConnection> ServerConnection;
 };
 
 IMPLEMENT_MODULE(FCookOnTheFlyModule, CookOnTheFly);
diff --git a/Engine/Source/Runtime/CookOnTheFly/Private/CookOnTheFlyPackageStore.cpp b/Engine/Source/Runtime/CookOnTheFly/Private/CookOnTheFlyPackageStore.cpp
deleted file mode 100644
diff --git a/Engine/Source/Runtime/CookOnTheFly/Private/CookOnTheFlyServerConnection.cpp b/Engine/Source/Runtime/CookOnTheFly/Private/CookOnTheFlyServerConnection.cpp
--- a/Engine/Source/Runtime/CookOnTheFly/Private/CookOnTheFlyServerConnection.cpp
+++ b/Engine/Source/Runtime/CookOnTheFly/Private/CookOnTheFlyServerConnection.cpp
@@ -53,9 +53,7 @@ public:
 		Disconnect();
 	}
 
-private:
-
-	virtual bool Connect(const UE::Cook::FCookOnTheFlyHostOptions& HostOptions) override
+	bool Connect(const UE::Cook::FCookOnTheFlyHostOptions& HostOptions)
 	{
 		check(HostOptions.Hosts.Num());
 
@@ -377,9 +375,14 @@ private:
 namespace UE { namespace Cook
 {
 
-TUniquePtr<ICookOnTheFlyServerConnection> MakeServerConnection()
+TUniquePtr<ICookOnTheFlyServerConnection> MakeServerConnection(const UE::Cook::FCookOnTheFlyHostOptions& HostOptions)
 {
-	return MakeUnique<FCookOnTheFlyServerConnection>();
+	TUniquePtr<FCookOnTheFlyServerConnection> Connection = MakeUnique<FCookOnTheFlyServerConnection>();
+	if (!Connection->Connect(HostOptions))
+	{
+		Connection.Release();
+	}
+	return Connection;
 }
 
 }} // namespace UE::Cook
diff --git a/Engine/Source/Runtime/CookOnTheFly/Public/CookOnTheFly.h b/Engine/Source/Runtime/CookOnTheFly/Public/CookOnTheFly.h
deleted file mode 100644
diff --git a/Engine/Source/Runtime/CookOnTheFly/Public/CookOnTheFlyPackageStore.h b/Engine/Source/Runtime/CookOnTheFly/Public/CookOnTheFlyPackageStore.h
deleted file mode 100644
diff --git a/Engine/Source/Runtime/Engine/Private/ODSC/ODSCThread.cpp b/Engine/Source/Runtime/Engine/Private/ODSC/ODSCThread.cpp
--- a/Engine/Source/Runtime/Engine/Private/ODSC/ODSCThread.cpp
+++ b/Engine/Source/Runtime/Engine/Private/ODSC/ODSCThread.cpp
@@ -7,7 +7,6 @@
 #include "HAL/PlatformProcess.h"
 #include "HAL/RunnableThread.h"
 #include "ShaderCompiler.h"
-#include "CookOnTheFly.h"
 
 FODSCRequestPayload::FODSCRequestPayload(EShaderPlatform InShaderPlatform, const FString& InMaterialName, const FString& InVertexFactoryName, const FString& InPipelineName, const TArray<FString>& InShaderTypeNames, const FString& InRequestHash)
 	: ShaderPlatform(InShaderPlatform), MaterialName(InMaterialName), VertexFactoryName(InVertexFactoryName), PipelineName(InPipelineName), ShaderTypeNames(std::move(InShaderTypeNames)), RequestHash(InRequestHash)
@@ -204,17 +203,7 @@ void FODSCThread::Process()
 	for (FODSCMessageHandler* NextRequest : RequestsToStart)
 	{
 		// send the info, the handler will process the response (and update shaders, etc)
-		UE::Cook::SendCookOnTheFlyRequest(
-			UE::Cook::ECookOnTheFlyMessage::RecompileShaders,
-			[NextRequest](FArchive& Request)
-			{
-				NextRequest->FillPayload(Request);
-			},
-			[NextRequest](FArchive& Response)
-			{
-				NextRequest->ProcessResponse(Response);
-				return true;
-			});
+		IFileManager::Get().SendMessageToServer(TEXT("RecompileShaders"), NextRequest);
 
 		CompletedThreadedRequests.Enqueue(NextRequest);
 	}
@@ -229,17 +218,7 @@ void FODSCThread::Process()
 		}
 
 		// send the info, the handler will process the response (and update shaders, etc)
-		UE::Cook::SendCookOnTheFlyRequest(
-			UE::Cook::ECookOnTheFlyMessage::RecompileShaders,
-			[RequestHandler](FArchive& Request)
-			{
-				RequestHandler->FillPayload(Request);
-			},
-			[RequestHandler](FArchive& Response)
-			{
-				RequestHandler->ProcessResponse(Response);
-				return true;
-			});
+		IFileManager::Get().SendMessageToServer(TEXT("RecompileShaders"), RequestHandler);
 
 		CompletedThreadedRequests.Enqueue(RequestHandler);
 	}
diff --git a/Engine/Source/Runtime/Launch/Private/LaunchEngineLoop.cpp b/Engine/Source/Runtime/Launch/Private/LaunchEngineLoop.cpp
--- a/Engine/Source/Runtime/Launch/Private/LaunchEngineLoop.cpp
+++ b/Engine/Source/Runtime/Launch/Private/LaunchEngineLoop.cpp
@@ -79,10 +79,6 @@
 #include "IO/IoDispatcher.h"
 #endif
 
-#if WITH_COTF
-#include "CookOnTheFly.h"
-#endif
-
 #if WITH_COREUOBJECT
 	#include "Internationalization/PackageLocalizationManager.h"
 	#include "Misc/PackageName.h"
@@ -1826,24 +1822,6 @@ int32 FEngineLoop::PreInitPreStartupScreen(const TCHAR* CmdLine)
 	}
 #endif // WITH_ENGINE
 
-#if WITH_COTF
-	if (IsRunningCookOnTheFly())
-	{
-		UE::Cook::FCookOnTheFlyHostOptions CookOnTheFlyHostOptions;
-		if (GetCookOnTheFlyHost(CookOnTheFlyHostOptions))
-		{
-			using namespace UE::Cook;
-			ICookOnTheFlyModule& CookOnTheFlyModule = FModuleManager::LoadModuleChecked<ICookOnTheFlyModule>(TEXT("CookOnTheFly"));
-			const bool bConnected = CookOnTheFlyModule.ConnectToServer(CookOnTheFlyHostOptions);
-			if (!bConnected)
-			{
-				UE_LOG(LogInit, Error, TEXT("Failed to connect to cook on the fly server"));
-				return 1;
-			}
-		}
-	}
-#endif // WITH_COTF
-
 	// allow the command line to override the platform file singleton
 	bool bFileOverrideFound = false;
 	{
diff --git a/Engine/Source/Runtime/StorageServerClient/Private/CookOnTheFlyPackageStore.cpp b/Engine/Source/Runtime/StorageServerClient/Private/CookOnTheFlyPackageStore.cpp
new file mode 100644
--- /dev/null
+++ b/Engine/Source/Runtime/StorageServerClient/Private/CookOnTheFlyPackageStore.cpp
@@ -0,0 +1,279 @@
+// Copyright Epic Games, Inc. All Rights Reserved.
+
+#include "CookOnTheFlyPackageStore.h"
+
+#if WITH_COTF
+
+#include "CookOnTheFly.h"
+#include "HAL/PlatformTime.h"
+#include "Containers/ChunkedArray.h"
+#include "Serialization/MemoryReader.h"
+#include "CookOnTheFlyMessages.h"
+
+class FCookOnTheFlyPackageStore final
+	: public FPackageStoreBase
+{
+public:
+	struct FEntryInfo
+	{
+		EPackageStoreEntryStatus Status = EPackageStoreEntryStatus::None;
+		int32 EntryIndex = INDEX_NONE;
+	};
+
+	struct FPackageStats
+	{
+		TAtomic<uint32> Cooked{ 0 };
+		TAtomic<uint32> Failed{ 0 };
+	};
+
+	FCookOnTheFlyPackageStore(UE::Cook::ICookOnTheFlyServerConnection& InCookOnTheFlyServerConnection)
+		: CookOnTheFlyServerConnection(InCookOnTheFlyServerConnection)
+	{
+		// Index zero is invalid
+		PackageEntries.Add();
+		CookOnTheFlyServerConnection.OnMessage().AddRaw(this, &FCookOnTheFlyPackageStore::OnCookOnTheFlyMessage);
+
+		using namespace UE::Cook;
+		using namespace UE::ZenCookOnTheFly::Messaging;
+
+		FCookOnTheFlyRequest Request(ECookOnTheFlyMessage::GetCookedPackages);
+		FCookOnTheFlyResponse Response = CookOnTheFlyServerConnection.SendRequest(Request).Get();
+
+		if (Response.IsOk())
+		{
+			FGetCookedPackagesResponse GetCookedPackagesResponse = Response.GetBodyAs<FGetCookedPackagesResponse>();
+			FPackageStoreData& PackageStoreData = GetCookedPackagesResponse.PackageStoreData;
+
+			UE_LOG(LogCookOnTheFly, Log, TEXT("Got '%d' cooked and '%d' failed packages from server"),
+				PackageStoreData.CookedPackages.Num(), PackageStoreData.FailedPackages.Num());
+
+			AddPackages(MoveTemp(PackageStoreData.CookedPackages), MoveTemp(PackageStoreData.FailedPackages));
+
+			LastServerActivtyTime = FPlatformTime::Seconds();
+		}
+		else
+		{
+			UE_LOG(LogCookOnTheFly, Warning, TEXT("Failed to send '%s' request"), LexToString(Request.GetHeader().MessageType));
+		}
+	}
+
+	virtual void Initialize() override
+	{
+
+	}
+
+	virtual void Lock() override
+	{
+	}
+
+	virtual void Unlock() override
+	{
+	}
+
+	virtual bool DoesPackageExist(FPackageId PackageId) override
+	{
+		FScopeLock _(&CriticalSection);
+		FEntryInfo EntryInfo = PackageIdToEntryInfo.FindRef(PackageId);
+		return EntryInfo.Status != EPackageStoreEntryStatus::Missing;
+	}
+
+	virtual EPackageStoreEntryStatus GetPackageStoreEntry(FPackageId PackageId, FPackageStoreEntry& OutPackageStoreEntry) override
+	{
+		using namespace UE::Cook;
+		using namespace UE::ZenCookOnTheFly::Messaging;
+
+		{
+			FScopeLock _(&CriticalSection);
+			FEntryInfo& EntryInfo = PackageIdToEntryInfo.FindOrAdd(PackageId, FEntryInfo());
+
+			if (EntryInfo.Status == EPackageStoreEntryStatus::Ok)
+			{
+				check(EntryInfo.EntryIndex != INDEX_NONE);
+				return CreatePackageStoreEntry(EntryInfo, OutPackageStoreEntry);
+			}
+			
+			if (EntryInfo.Status == EPackageStoreEntryStatus::Missing)
+			{
+				return EPackageStoreEntryStatus::Missing;
+			}
+
+			if (EntryInfo.Status == EPackageStoreEntryStatus::Pending)
+			{
+				CheckActivity();
+				return EPackageStoreEntryStatus::Pending;
+			}
+
+			// The package hasn't been requested, set the status to pending.
+			EntryInfo.Status = EPackageStoreEntryStatus::Pending;
+		}
+
+		LastClientActivtyTime = FPlatformTime::Seconds();
+		UE_LOG(LogCookOnTheFly, Verbose, TEXT("Requesting package 0x%llX"), PackageId.ValueForDebugging());
+
+		FCookOnTheFlyRequest Request(ECookOnTheFlyMessage::CookPackage);
+		Request.SetBodyTo(FCookPackageRequest { PackageId });
+		FCookOnTheFlyResponse Response = CookOnTheFlyServerConnection.SendRequest(Request).Get();
+
+		if (!Response.IsOk())
+		{
+			UE_LOG(LogCookOnTheFly, Warning, TEXT("Failed to send '%s' request"), LexToString(Request.GetHeader().MessageType));
+			FScopeLock _(&CriticalSection);
+			FEntryInfo& EntryInfo = PackageIdToEntryInfo.FindChecked(PackageId);
+			EntryInfo.Status = EPackageStoreEntryStatus::Missing;
+			return EPackageStoreEntryStatus::Missing;
+		}
+
+		FCookPackageResponse CookPackageResponse = Response.GetBodyAs<FCookPackageResponse>();
+
+		{
+			FScopeLock _(&CriticalSection);
+
+			FEntryInfo& EntryInfo = PackageIdToEntryInfo.FindChecked(PackageId);
+			if (CookPackageResponse.Status == EPackageStoreEntryStatus::Missing)
+			{
+				EntryInfo.Status = EPackageStoreEntryStatus::Missing;
+				return EPackageStoreEntryStatus::Missing;
+			}
+			else
+			{
+				return CreatePackageStoreEntry(EntryInfo, OutPackageStoreEntry);
+			}
+		}
+	}
+
+	EPackageStoreEntryStatus CreatePackageStoreEntry(const FEntryInfo& EntryInfo, FPackageStoreEntry& OutPackageStoreEntry)
+	{
+		if (EntryInfo.Status == EPackageStoreEntryStatus::Ok)
+		{
+			const FPackageStoreEntryResource& Entry = PackageEntries[EntryInfo.EntryIndex];
+			OutPackageStoreEntry.ExportInfo = Entry.ExportInfo;
+			OutPackageStoreEntry.ImportedPackageIds = Entry.ImportedPackageIds;
+			return EPackageStoreEntryStatus::Ok;
+		}
+		else
+		{
+			return EntryInfo.Status;
+		}
+	}
+
+	virtual bool GetPackageRedirectInfo(FPackageId PackageId, FName& OutSourcePackageName, FPackageId& OutRedirectedToPackageId)  override
+	{
+		return false;
+	}
+
+	void AddPackages(TArray<FPackageStoreEntryResource> Entries, TArray<FPackageId> FailedPackageIds)
+	{
+		FScopeLock _(&CriticalSection);
+		
+		for (FPackageId FailedPackageId : FailedPackageIds)
+		{
+			UE_LOG(LogCookOnTheFly, Verbose, TEXT("'0x%llX' [Failed]"), FailedPackageId.Value());
+			FEntryInfo& EntryInfo = PackageIdToEntryInfo.FindOrAdd(FailedPackageId, FEntryInfo());
+			EntryInfo.Status = EPackageStoreEntryStatus::Missing;
+			PackageStats.Failed++;
+		}
+
+		const int32 NumPackageEntries = PackageEntries.Num();
+
+		for (FPackageStoreEntryResource& Entry : Entries)
+		{
+			const FPackageId PackageId = Entry.GetPackageId();
+			FEntryInfo& EntryInfo = PackageIdToEntryInfo.FindOrAdd(PackageId, FEntryInfo());
+
+			if (EntryInfo.EntryIndex == INDEX_NONE)
+			{
+				EntryInfo.Status = EPackageStoreEntryStatus::Ok;
+				EntryInfo.EntryIndex = PackageEntries.Add();
+				PackageEntries[EntryInfo.EntryIndex] = MoveTemp(Entry);
+				PackageStats.Cooked++;
+
+				UE_LOG(LogCookOnTheFly, Verbose, TEXT("'%s' [OK] (Cooked/Failed='%d/%d')"),
+					*PackageEntries[EntryInfo.EntryIndex].PackageName.ToString(), PackageStats.Cooked.Load(), PackageStats.Failed.Load());
+			}
+		}
+	}
+
+	void OnCookOnTheFlyMessage(const UE::Cook::FCookOnTheFlyMessage& Message)
+	{
+		using namespace UE::ZenCookOnTheFly::Messaging;
+
+		switch (Message.GetHeader().MessageType)
+		{
+			case UE::Cook::ECookOnTheFlyMessage::PackagesCooked:
+			{
+				FPackagesCookedMessage PackagesCookedMessage = Message.GetBodyAs<FPackagesCookedMessage>();
+				FPackageStoreData& PackageStoreData = PackagesCookedMessage.PackageStoreData;
+
+				UE_LOG(LogCookOnTheFly, Verbose, TEXT("Received '%s' message, Cooked='%d', Failed='%d', Server total='%d/%d' (Cooked/Failed)"),
+					LexToString(Message.GetHeader().MessageType),
+					PackageStoreData.CookedPackages.Num(),
+					PackageStoreData.FailedPackages.Num(),
+					PackageStoreData.TotalCookedPackages,
+					PackageStoreData.TotalFailedPackages);
+
+				AddPackages(MoveTemp(PackageStoreData.CookedPackages), MoveTemp(PackageStoreData.FailedPackages));
+
+				UE_CLOG(PackageStoreData.TotalCookedPackages != PackageStats.Cooked.Load() || PackageStoreData.TotalFailedPackages != PackageStats.Failed.Load(),
+					LogCookOnTheFly, Warning, TEXT("Client/Server package mismatch, Cooked='%d/%d', Failed='%d/%d' (Client/Server)"),
+					PackageStats.Cooked.Load(), PackageStoreData.TotalCookedPackages, PackageStats.Failed.Load(), PackageStoreData.TotalFailedPackages);
+
+				PendingEntriesAdded.Broadcast();
+
+				break;
+			}
+
+			default:
+				break;
+		}
+
+		LastServerActivtyTime = FPlatformTime::Seconds();
+	}
+
+	void CheckActivity()
+	{
+		const double TimeSinceLastClientActivity = FPlatformTime::Seconds() - LastClientActivtyTime;
+		const double TimeSinceLastServerActivity = FPlatformTime::Seconds() - LastServerActivtyTime;
+		const double TimeSinceLastWarning = FPlatformTime::Seconds() - LastWarningTime;
+
+		if (TimeSinceLastClientActivity > MaxInactivityTime &&
+			TimeSinceLastServerActivity > MaxInactivityTime &&
+			TimeSinceLastWarning > TimeBetweenWarning)
+		{
+			LastWarningTime = FPlatformTime::Seconds();
+
+			UE_LOG(LogCookOnTheFly, Log, TEXT("No server response in '%.2lf' seconds"), TimeSinceLastServerActivity);
+
+			UE_LOG(LogCookOnTheFly, Log, TEXT("=== Pending Packages ==="));
+			{
+				FScopeLock _(&CriticalSection);
+				for (const auto& KeyValue : PackageIdToEntryInfo)
+				{
+					const FEntryInfo& EntryInfo = KeyValue.Value;
+					if (EntryInfo.Status == EPackageStoreEntryStatus::Pending)
+					{
+						UE_LOG(LogCookOnTheFly, Log, TEXT("0x%llX"), KeyValue.Key.ValueForDebugging());
+					}
+				}
+			}
+		}
+	}
+
+	UE::Cook::ICookOnTheFlyServerConnection& CookOnTheFlyServerConnection;
+	FCriticalSection CriticalSection;
+	TMap<FPackageId, FEntryInfo> PackageIdToEntryInfo;
+	TChunkedArray<FPackageStoreEntryResource> PackageEntries;
+	FPackageStats PackageStats;
+
+	const double MaxInactivityTime = 20;
+	const double TimeBetweenWarning = 10;
+	double LastClientActivtyTime = 0;
+	double LastServerActivtyTime = 0;
+	double LastWarningTime = 0;
+};
+
+TSharedPtr<IPackageStore> MakeCookOnTheFlyPackageStore(UE::Cook::ICookOnTheFlyServerConnection& CookOnTheFlyServerConnection)
+{
+	return MakeShared<FCookOnTheFlyPackageStore>(CookOnTheFlyServerConnection);
+}
+
+#endif // WITH_COTF
diff --git a/Engine/Source/Runtime/StorageServerClient/Private/CookOnTheFlyPackageStore.h b/Engine/Source/Runtime/StorageServerClient/Private/CookOnTheFlyPackageStore.h
new file mode 100644
--- /dev/null
+++ b/Engine/Source/Runtime/StorageServerClient/Private/CookOnTheFlyPackageStore.h
@@ -0,0 +1,20 @@
+// Copyright Epic Games, Inc. All Rights Reserved.
+
+#pragma once
+
+#include "CoreMinimal.h"
+
+#if WITH_COTF
+
+#include "IO/PackageStore.h"
+
+class FIoDispatcher;
+
+namespace UE { namespace Cook
+{
+	class ICookOnTheFlyServerConnection;
+}}
+
+TSharedPtr<IPackageStore> MakeCookOnTheFlyPackageStore(UE::Cook::ICookOnTheFlyServerConnection& CookOnTheFlyServerConnection);
+
+#endif // WITH_COTF
diff --git a/Engine/Source/Runtime/StorageServerClient/Private/StorageServerPlatformFile.cpp b/Engine/Source/Runtime/StorageServerClient/Private/StorageServerPlatformFile.cpp
--- a/Engine/Source/Runtime/StorageServerClient/Private/StorageServerPlatformFile.cpp
+++ b/Engine/Source/Runtime/StorageServerClient/Private/StorageServerPlatformFile.cpp
@@ -289,15 +289,37 @@ void FStorageServerPlatformFile::InitializeAfterProjectFilePath()
 			FIoDispatcher& IoDispatcher = FIoDispatcher::Get();
 			TSharedRef<FStorageServerIoDispatcherBackend> IoDispatcherBackend = MakeShared<FStorageServerIoDispatcherBackend>(*Connection.Get());
 			IoDispatcher.Mount(IoDispatcherBackend);
+
+#if WITH_COTF
+			if (FParse::Param(FCommandLine::Get(), TEXT("CookOnTheFly")))
+			{
+				UE::Cook::FCookOnTheFlyHostOptions CookOnTheFlyHostOptions;
+				// Cook-on-the-fly expects the same host as the Zen storage server
+				CookOnTheFlyHostOptions.Hosts = HostAddrs;
+				double ServerWaitTimeInSeconds;
+				if (FParse::Value(FCommandLine::Get(), TEXT("-CookOnTheFlyServerWaitTime="), ServerWaitTimeInSeconds))
+				{
+					CookOnTheFlyHostOptions.ServerStartupWaitTime = FTimespan::FromSeconds(ServerWaitTimeInSeconds);
+				}
+				UE::Cook::ICookOnTheFlyModule& CookOnTheFlyModule = FModuleManager::LoadModuleChecked<UE::Cook::ICookOnTheFlyModule>(TEXT("CookOnTheFly"));
+				CookOnTheFlyServerConnection = CookOnTheFlyModule.ConnectToServer(CookOnTheFlyHostOptions);
+				if (CookOnTheFlyServerConnection)
+				{
+					CookOnTheFlyServerConnection->OnMessage().AddRaw(this, &FStorageServerPlatformFile::OnCookOnTheFlyMessage);
+				}
+				else
+				{
+					UE_LOG(LogStorageServerPlatformFile, Fatal, TEXT("Failed to connect to cook on the fly server"));
+				}
+			}
+#endif
+
 			FCoreDelegates::CreatePackageStore.BindLambda([this]() -> TSharedPtr<IPackageStore>
 			{
 #if WITH_COTF
-				if (IsRunningCookOnTheFly())
+				if (CookOnTheFlyServerConnection)
 				{
-					UE::Cook::ICookOnTheFlyModule& CookOnTheFlyModule = FModuleManager::LoadModuleChecked<UE::Cook::ICookOnTheFlyModule>(TEXT("CookOnTheFly"));
-					UE::Cook::ICookOnTheFlyServerConnection& CotfConnection = CookOnTheFlyModule.GetServerConnection();
-					CotfConnection.OnMessage().AddRaw(this, &FStorageServerPlatformFile::OnCookOnTheFlyMessage);
-					return MakeCookOnTheFlyPackageStore(CotfConnection);
+					return MakeCookOnTheFlyPackageStore(*CookOnTheFlyServerConnection.Get());
 				}
 				else
 #endif
@@ -695,6 +717,30 @@ int64 FStorageServerPlatformFile::SendReadMessage(uint8* Destination, const FIoC
 	return BytesRead;
 }
 
+bool FStorageServerPlatformFile::SendMessageToServer(const TCHAR* Message, IPlatformFile::IFileServerMessageHandler* Handler)
+{
+#if WITH_COTF
+	if (FCString::Stricmp(Message, TEXT("RecompileShaders")) == 0)
+	{
+		UE::Cook::FCookOnTheFlyRequest Request(UE::Cook::ECookOnTheFlyMessage::RecompileShaders);
+		{
+			TUniquePtr<FArchive> Ar = Request.WriteBody();
+			Handler->FillPayload(*Ar);
+		}
+
+		UE::Cook::FCookOnTheFlyResponse Response = CookOnTheFlyServerConnection->SendRequest(Request).Get();
+		if (Response.IsOk())
+		{
+			TUniquePtr<FArchive> Ar = Response.ReadBody();
+			Handler->ProcessResponse(*Ar);
+		}
+
+		return Response.IsOk();
+	}
+#endif
+	return false;
+}
+
 #if WITH_COTF
 void FStorageServerPlatformFile::OnCookOnTheFlyMessage(const UE::Cook::FCookOnTheFlyMessage& Message)
 {
diff --git a/Engine/Source/Runtime/StorageServerClient/Private/StorageServerPlatformFile.h b/Engine/Source/Runtime/StorageServerClient/Private/StorageServerPlatformFile.h
--- a/Engine/Source/Runtime/StorageServerClient/Private/StorageServerPlatformFile.h
+++ b/Engine/Source/Runtime/StorageServerClient/Private/StorageServerPlatformFile.h
@@ -17,6 +17,7 @@ class IPackageStore;
 namespace UE::Cook
 {
 	class FCookOnTheFlyMessage;
+	class ICookOnTheFlyServerConnection;
 }
 #endif
 
@@ -95,6 +96,7 @@ public:
 	virtual IFileHandle* OpenWrite(const TCHAR* Filename, bool bAppend = false, bool bAllowRead = false) override;
 	virtual bool CreateDirectory(const TCHAR* Directory) override;
 	virtual bool DeleteDirectory(const TCHAR* Directory) override;
+	virtual bool SendMessageToServer(const TCHAR* Message, IPlatformFile::IFileServerMessageHandler* Handler) override;
 
 private:
 	friend class FStorageServerFileHandle;
@@ -113,6 +115,9 @@ private:
 	FStringView ServerEngineDirView = FStringView(TEXT("/{engine}/"));
 	FStringView ServerProjectDirView = FStringView(TEXT("/{project}/"));
 	TUniquePtr<FStorageServerConnection> Connection;
+#if WITH_COTF
+	TUniquePtr<UE::Cook::ICookOnTheFlyServerConnection> CookOnTheFlyServerConnection;
+#endif
 	FStorageServerFileSystemTOC ServerToc;
 	FString ServerProject;
 	FString ServerPlatform;
