Zen COTF: Fixed issue with FCookOnTheFlyServerConnection::IsConnected always returning true after a disconnect
#rnx

diff --git a/Engine/Source/Runtime/CookOnTheFly/Private/CookOnTheFlyServerConnection.cpp b/Engine/Source/Runtime/CookOnTheFly/Private/CookOnTheFlyServerConnection.cpp
--- a/Engine/Source/Runtime/CookOnTheFly/Private/CookOnTheFlyServerConnection.cpp
+++ b/Engine/Source/Runtime/CookOnTheFly/Private/CookOnTheFlyServerConnection.cpp
@@ -173,7 +173,7 @@ public:
 				Thread.Wait();
 				Thread.Reset();
 				Socket.Reset();
-				ClientId  = -1;
+				ClientId  = 0;
 			}
 		}
 	}
@@ -327,6 +327,7 @@ public:
 		}
 
 		UE_LOG(LogCotfServerConnection, Display, TEXT("Terminating connection to server '%s'"), *ServerAddr->ToString(true));
+		ClientId = 0;
 	}
 
 	struct FPendingRequest
