Compile fix for editor in shipping config
#rnx

diff --git a/Engine/Source/Editor/UnrealEd/Private/UnrealEdEngine.cpp b/Engine/Source/Editor/UnrealEd/Private/UnrealEdEngine.cpp
--- a/Engine/Source/Editor/UnrealEd/Private/UnrealEdEngine.cpp
+++ b/Engine/Source/Editor/UnrealEd/Private/UnrealEdEngine.cpp
@@ -222,7 +222,9 @@ void UUnrealEdEngine::Init(IEngineLoop* InEngineLoop)
 			CookServer->Initialize(ECookMode::CookByTheBookFromTheEditor, BaseCookingFlags);
 		}
 
+#if WITH_COTF
 		ExternalCookOnTheFlyServer = new FExternalCookOnTheFlyServer();
+#endif
 	}
 
 	if (FParse::Param(FCommandLine::Get(), TEXT("nomcp")))
diff --git a/Engine/Source/Runtime/CookOnTheFly/Internal/CookOnTheFlyMessages.h b/Engine/Source/Runtime/CookOnTheFly/Internal/CookOnTheFlyMessages.h
--- a/Engine/Source/Runtime/CookOnTheFly/Internal/CookOnTheFlyMessages.h
+++ b/Engine/Source/Runtime/CookOnTheFly/Internal/CookOnTheFlyMessages.h
@@ -3,9 +3,6 @@
 #pragma once
 
 #include "CoreMinimal.h"
-
-#if WITH_COTF
-
 #include "IO/PackageStore.h"
 
 namespace UE { namespace ZenCookOnTheFly { namespace Messaging
@@ -47,5 +44,3 @@ struct FRecookPackagesResponse
 };
 
 }}} // namesapce UE::ZenCookOnTheFly::Messaging
-
-#endif // WITH_COTF
diff --git a/Engine/Source/Runtime/CookOnTheFly/Private/CookOnTheFlyMessages.cpp b/Engine/Source/Runtime/CookOnTheFly/Private/CookOnTheFlyMessages.cpp
--- a/Engine/Source/Runtime/CookOnTheFly/Private/CookOnTheFlyMessages.cpp
+++ b/Engine/Source/Runtime/CookOnTheFly/Private/CookOnTheFlyMessages.cpp
@@ -2,8 +2,6 @@
 
 #include "CookOnTheFlyMessages.h"
 
-#if WITH_COTF
-
 namespace UE { namespace ZenCookOnTheFly { namespace Messaging
 {
 
@@ -46,5 +44,3 @@ FArchive& operator<<(FArchive& Ar, FRecookPackagesResponse& Response)
 }
 
 }}} // namesapce UE::ZenCookOnTheFly::Messaging
-
-#endif // WITH_COTF
