ZenCOTF: Set NoDelay on network sockets
#rnx

diff --git a/Engine/Source/Editor/UnrealEd/Private/Cooker/IoStoreCookOnTheFlyRequestManager.cpp b/Engine/Source/Editor/UnrealEd/Private/Cooker/IoStoreCookOnTheFlyRequestManager.cpp
--- a/Engine/Source/Editor/UnrealEd/Private/Cooker/IoStoreCookOnTheFlyRequestManager.cpp
+++ b/Engine/Source/Editor/UnrealEd/Private/Cooker/IoStoreCookOnTheFlyRequestManager.cpp
@@ -98,6 +98,7 @@ public:
 		}
 
 		Socket->SetReuseAddr();
+		Socket->SetNoDelay();
 
 		if (!Socket->Bind(*ListenAddr))
 		{
diff --git a/Engine/Source/Runtime/CookOnTheFly/Private/CookOnTheFlyServerConnection.cpp b/Engine/Source/Runtime/CookOnTheFly/Private/CookOnTheFlyServerConnection.cpp
--- a/Engine/Source/Runtime/CookOnTheFly/Private/CookOnTheFlyServerConnection.cpp
+++ b/Engine/Source/Runtime/CookOnTheFly/Private/CookOnTheFlyServerConnection.cpp
@@ -55,6 +55,7 @@ public:
 
 	bool Connect(const UE::Cook::FCookOnTheFlyHostOptions& HostOptions)
 	{
+		TRACE_CPUPROFILER_EVENT_SCOPE(FCookOnTheFlyServerConnection::Connect);
 		check(HostOptions.Hosts.Num());
 
 		const int32 Port = HostOptions.Port > 0 ? HostOptions.Port : UE::Cook::DefaultCookOnTheFlyServingPort;
@@ -78,12 +79,18 @@ public:
 				UE_LOG(LogCotfServerConnection, Display, TEXT("Connecting to COTF server at '%s'..."), *Addr->ToString(true));
 			
 				Socket.Reset(SocketSubsystem.CreateSocket(NAME_Stream, TEXT("COTF-ServerConnection"), Addr->GetProtocolType()));
-				if (Socket.IsValid() && Socket->Connect(*Addr))
+				if (!Socket.IsValid())
 				{
-					ServerAddr = Addr;
-					bConnected = true;
-					break;
+					continue;
+				}
+				Socket->SetNoDelay();
+				if (!Socket->Connect(*Addr))
+				{
+					continue;
 				}
+				ServerAddr = Addr;
+				bConnected = true;
+				break;
 			}
 
 			if (bConnected || FPlatformTime::Seconds() > ServerWaitEndTime)
