Fixed a crash after a failed request lookup in the CotF server.

diff --git a/Engine/Source/Runtime/CookOnTheFly/Private/CookOnTheFlyServerConnection.cpp b/Engine/Source/Runtime/CookOnTheFly/Private/CookOnTheFlyServerConnection.cpp
--- a/Engine/Source/Runtime/CookOnTheFly/Private/CookOnTheFlyServerConnection.cpp
+++ b/Engine/Source/Runtime/CookOnTheFly/Private/CookOnTheFlyServerConnection.cpp
@@ -211,7 +211,12 @@ private:
 			else if (bIsResponse)
 			{
 				FPendingRequest* PendingRequest = GetRequest(MessageHeader.CorrelationId);
-				check(PendingRequest);
+				if (!PendingRequest)
+				{
+					UE_LOG(LogCotfServerConnection, Warning, TEXT("Failed to match response with id '%u' to an existing request"), MessageHeader.CorrelationId);
+					return false;
+				}
+
 				check(PendingRequest->RequestHeader.CorrelationId == MessageHeader.CorrelationId);
 
 				FCookOnTheFlyResponse Response;
