Fix: Deadlock when losing connection to ODSC server

As suggested by [REDACTED_HANDLE], clean up pending requests when the COTF connection thread is about to die.

Because our TPromise can't be abandoned, we need to assign some value to the Promises in the PendingRequests before we can destroy them. I chose to assign an error status response, similar to the response from a failed SendRequest. Since we assign values to the Promises, any threads waiting on those Promises' Futures are awakened. I also removed the assert failure on error response in ODSCThread and changed it to a log message and server connection reset, so that a future ODSC request would attempt to reestablish a connection to the server.

diff --git a/Engine/Source/Runtime/CookOnTheFly/Private/CookOnTheFlyServerConnection.cpp b/Engine/Source/Runtime/CookOnTheFly/Private/CookOnTheFlyServerConnection.cpp
--- a/Engine/Source/Runtime/CookOnTheFly/Private/CookOnTheFlyServerConnection.cpp
+++ b/Engine/Source/Runtime/CookOnTheFly/Private/CookOnTheFlyServerConnection.cpp
@@ -282,6 +282,7 @@ private:
 			});
 
 		UE_LOG(LogCotfServerConnection, Display, TEXT("Terminating connection to server"));
+		Flush();
 		Transport.Reset();
 	}
 
@@ -311,6 +312,25 @@ private:
 		FScopeLock _(&RequestsCriticalSection);
 		PendingRequests.Remove(PendingRequest->RequestHeader.CorrelationId);
 	}
+
+	void Flush()
+	{
+		if (!PendingRequests.IsEmpty())
+		{
+			using namespace UE::Cook;
+			FCookOnTheFlyRequest SyntheticRequest;
+			FScopeLock _(&RequestsCriticalSection);
+			for (const TPair<uint32, TUniquePtr<FPendingRequest>>& Pair : PendingRequests)
+			{
+				FPendingRequest& PendingRequest = *Pair.Value;
+				SyntheticRequest.SetHeader(PendingRequest.RequestHeader);
+				FCookOnTheFlyResponse ErrorResponse(SyntheticRequest);
+				ErrorResponse.SetStatus(ECookOnTheFlyMessageStatus::Error);
+				PendingRequest.ResponsePromise.SetValue(ErrorResponse);
+			}
+			PendingRequests.Empty();
+		}
+	}
 	
 	FPendingRequest* GetRequest(uint32 CorrelationId)
 	{
diff --git a/Engine/Source/Runtime/Engine/Private/ODSC/ODSCThread.cpp b/Engine/Source/Runtime/Engine/Private/ODSC/ODSCThread.cpp
--- a/Engine/Source/Runtime/Engine/Private/ODSC/ODSCThread.cpp
+++ b/Engine/Source/Runtime/Engine/Private/ODSC/ODSCThread.cpp
@@ -324,6 +324,9 @@ void FODSCThread::SendMessageToServer(IPlatformFile::IFileServerMessageHandler*
 		TUniquePtr<FArchive> Ar = Response.ReadBody();
 		Handler->ProcessResponse(*Ar);
 	}
-
-	check(Response.IsOk());
+	else
+	{
+		UE_LOG(LogODSC, Display, TEXT("Received error response from CookOnTheFlyServerConnection; disconnecting"));
+		CookOnTheFlyServerConnection.Reset();
+	}
 }
