Fix Mismatch between platform requested by client and ODSC server

Unfortunately, ODSC/COTF server (a.k.a. UE Editor) and game/client runtime don't generally agree on the name of the target platform for client-only target configurations. Editor target platform names for client-only target configs always end with "Client," but runtime platform names don't universally adhere to that. Desktop platforms do, but non-desktop platforms don't.

As suggested by [REDACTED_HANDLE], add "Client" to the end of the requested platform name(s) from the runtime (client) side of CookOnTheFlyServerConnection, if built for a client-only configuration and not already present, in order to match up with the server's notion of the platform name. [REDACTED_HANDLE] suggested adding Client (and Server) platforms to the array generated by FPlatformMisc::GetValidTargetPlatforms(), and while that would feel cleaner, I deemed it more risky and complicated than was warranted.

diff --git a/Engine/Source/Runtime/CookOnTheFly/Private/CookOnTheFlyServerConnection.cpp b/Engine/Source/Runtime/CookOnTheFly/Private/CookOnTheFlyServerConnection.cpp
--- a/Engine/Source/Runtime/CookOnTheFly/Private/CookOnTheFlyServerConnection.cpp
+++ b/Engine/Source/Runtime/CookOnTheFly/Private/CookOnTheFlyServerConnection.cpp
@@ -32,7 +32,7 @@ public:
 		TArray<FString> TargetPlatformNames;
 		if (FPlatformProperties::RequiresCookedData())
 		{
-			FPlatformMisc::GetValidTargetPlatforms(TargetPlatformNames);
+			GetCookOnTheFlyTargetPlatforms(TargetPlatformNames);
 		}
 		
 		FBufferArchive HandshakeRequestPayload;
@@ -198,6 +198,24 @@ public:
 	}
 
 private:
+	static void GetCookOnTheFlyTargetPlatforms(TArray<FString>& TargetPlatformNames)
+	{
+		FPlatformMisc::GetValidTargetPlatforms(TargetPlatformNames);
+		// COTF server (a.k.a. UE Editor) appends "Client" to platform names for client-only targets
+		// Desktop platforms do the same for runtime client config, but non-desktop platforms don't
+		// Append "Client" to platform names, if not already present, in client-only configs to match COTF server
+		if (FPlatformProperties::IsClientOnly())
+		{
+			for (FString& TargetPlatformName : TargetPlatformNames)
+			{
+				if (!TargetPlatformName.Contains(TEXT("Client")))
+				{
+					TargetPlatformName.Append(TEXT("Client"));
+				}
+			}
+		}
+	}
+
 	void Disconnect()
 	{
 		if (!bStopRequested.Exchange(true))
