Fix platformProtocol when using OSDC (CookOnTheFly). Condition was removed in CL 40137015.

diff --git a/Engine/Source/Runtime/CookOnTheFly/Private/CookOnTheFlyModule.cpp b/Engine/Source/Runtime/CookOnTheFly/Private/CookOnTheFlyModule.cpp
--- a/Engine/Source/Runtime/CookOnTheFly/Private/CookOnTheFlyModule.cpp
+++ b/Engine/Source/Runtime/CookOnTheFly/Private/CookOnTheFlyModule.cpp
@@ -127,7 +127,7 @@ private:
 			ServerStartupWaitTime = FTimespan::FromSeconds(ServerWaitTimeInSeconds);
 		}
 
-		if (FParse::Param(FCommandLine::Get(), TEXT("ZenPlatformProtocol")))
+		if (FParse::Param(FCommandLine::Get(), TEXT("ZenPlatformProtocol")) || FParse::Param(FCommandLine::Get(), TEXT("PlatformProtocol")))
 		{
 			HostIpList.Add("platform://");
 		}
